Moved GameplayBehaviors out of restricted folder over to Experimental

diff --git a/Engine/Plugins/Runtime/ZoneGraph/Source/ZoneGraph/Public/ZoneGraphBuilder.h b/Engine/Plugins/Runtime/ZoneGraph/Source/ZoneGraph/Public/ZoneGraphBuilder.h
new file mode 100644
--- /dev/null
+++ b/Engine/Plugins/Runtime/ZoneGraph/Source/ZoneGraph/Public/ZoneGraphBuilder.h
@@ -0,0 +1,119 @@
+// Copyright Epic Games, Inc. All Rights Reserved.
+
+#pragma once
+
+#include "CoreMinimal.h"
+#include "ZoneGraphTypes.h"
+#include "HierarchicalHashGrid2D.h"
+#include "ZoneGraphBuilder.generated.h"
+
+class AZoneGraphData;
+class UZoneShapeComponent;
+struct FZoneShapeLaneInternalLink;
+
+typedef THierarchicalHashGrid2D<5, 4> FZoneGraphBuilderHashGrid2D;	// 5 levels of hierarchy, 4 ration between levels (biggest bucket 4^5 = 1024 cells)
+
+USTRUCT()
+struct ZONEGRAPH_API FZoneGraphBuilderRegisteredComponent
+{
+	GENERATED_BODY()
+
+	FZoneGraphBuilderRegisteredComponent() = default;
+	FZoneGraphBuilderRegisteredComponent(UZoneShapeComponent* InComponent) : Component(InComponent) {}
+
+	UPROPERTY()
+	UZoneShapeComponent* Component = nullptr;
+
+	UPROPERTY()
+	uint32 ShapeHash = 0;
+
+	FZoneGraphBuilderHashGrid2D::FCellLocation CellLoc;
+};
+
+// Build time data for UZoneShapeComponent
+USTRUCT()
+struct ZONEGRAPH_API FZoneShapeComponentBuildData
+{
+	GENERATED_BODY()
+
+	UPROPERTY()
+	int32 ZoneIndex = 0;
+
+	UPROPERTY()
+	TArray<FZoneGraphLaneHandle> Lanes;
+};
+
+// Build time data, can be used to map things between editor representation and baked data.
+USTRUCT()
+struct ZONEGRAPH_API FZoneGraphBuildData
+{
+	GENERATED_BODY()
+
+	void Reset()
+	{
+		ZoneShapeComponentBuildData.Reset();
+	}
+
+	UPROPERTY()
+	TMap<const UZoneShapeComponent*, FZoneShapeComponentBuildData> ZoneShapeComponentBuildData;
+};
+
+USTRUCT()
+struct ZONEGRAPH_API FZoneGraphBuilder
+{
+	GENERATED_BODY()
+
+public:
+	FZoneGraphBuilder();
+	~FZoneGraphBuilder();
+
+	void RegisterZoneShapeComponent(UZoneShapeComponent& ShapeComp);
+	void UnregisterZoneShapeComponent(UZoneShapeComponent& ShapeComp);
+	void OnZoneShapeComponentChanged(UZoneShapeComponent& ShapeComp);
+	const TArray<FZoneGraphBuilderRegisteredComponent>& GetRegisteredZoneShapeComponents() const { return ShapeComponents; }
+
+	bool NeedsRebuild() const { return bIsDirty; }
+
+	/** Builds zone graph for each zone graph data based on registered shapes.
+	 * @param AllZoneGraphData All Zone Graph data to update.
+	 * @param bForceRebuild If set will build graph even if it inputs have not changed.
+	 */
+	void BuildAll(const TArray<AZoneGraphData*>& AllZoneGraphData, const bool bForceRebuild);
+
+	void FindShapeConnections(const UZoneShapeComponent& SourceShapeComp, TArray<FZoneShapeConnection>& OutShapeConnections) const;
+
+	/** Converts single zone shape into a zone storage, used in UI for editing and rendering.
+	*/
+	static void BuildSingleShape(const UZoneShapeComponent& ShapeComp, const FMatrix& LocalToWorld, FZoneGraphStorage& OutZoneStorage);
+
+protected:
+	void Build(AZoneGraphData& ZoneGraphData);
+	void RequestRebuild();
+	void OnLaneProfileChanged(const FZoneLaneProfileRef& ChangedLaneProfileRef);
+	uint32 CalculateCombinedShapeHash(const AZoneGraphData& ZoneGraphData) const;
+	static void AppendShapeToZoneStorage(const UZoneShapeComponent& ShapeComp, const FMatrix& LocalToWorld, FZoneGraphStorage& OutZoneStorage, TArray<FZoneShapeLaneInternalLink>& OutInternalLinks, FZoneGraphBuildData* InBuildData = nullptr);
+	static void ConnectLanes(TArray<FZoneShapeLaneInternalLink>& InternalLinks, FZoneGraphStorage& OutZoneStorage);
+
+	UPROPERTY(Transient)
+	TArray<FZoneGraphBuilderRegisteredComponent> ShapeComponents;
+
+	UPROPERTY(Transient)
+	TArray<int32> ShapeComponentsFreeList;
+
+	UPROPERTY(Transient)
+	TMap<UZoneShapeComponent*, int32> ShapeComponentToIndex;
+
+	UPROPERTY(Transient)
+	FZoneGraphBuildData BuildData;
+
+#if WITH_EDITOR
+	FDelegateHandle OnBuildSettingsChangedHandle;
+	FDelegateHandle OnTagsChangedHandle;
+	FDelegateHandle OnLaneProfileChangedHandle;
+#endif
+
+	FZoneGraphBuilderHashGrid2D HashGrid;
+
+	bool bSkipHashCheck = false;
+	bool bIsDirty = false;
+};
diff --git a/Engine/Plugins/Runtime/ZoneGraph/Source/ZoneGraph/Public/ZoneGraphData.h b/Engine/Plugins/Runtime/ZoneGraph/Source/ZoneGraph/Public/ZoneGraphData.h
new file mode 100644
--- /dev/null
+++ b/Engine/Plugins/Runtime/ZoneGraph/Source/ZoneGraph/Public/ZoneGraphData.h
@@ -0,0 +1,75 @@
+// Copyright Epic Games, Inc. All Rights Reserved.
+
+#pragma once
+
+#include "GameFramework/Actor.h"
+#include "ZoneGraphTypes.h"
+#include "ZoneGraphData.generated.h"
+
+class UZoneGraphRenderingComponent;
+
+UCLASS(config = ZoneGraph, defaultconfig, NotBlueprintable)
+class ZONEGRAPH_API AZoneGraphData : public AActor
+{
+	GENERATED_BODY()
+public:
+	AZoneGraphData(const FObjectInitializer& ObjectInitializer = FObjectInitializer::Get());
+
+	//~ Begin UObject/AActor Interface
+	virtual void PostActorCreated() override;
+	virtual void PostLoad() override;
+	virtual void Destroyed() override;
+	virtual void EndPlay(const EEndPlayReason::Type EndPlayReason) override;
+	virtual void PreRegisterAllComponents() override;
+	virtual void PostUnregisterAllComponents() override;
+
+#if WITH_EDITOR
+	virtual void PostEditUndo() override;
+	virtual void PostEditChangeProperty(FPropertyChangedEvent& PropertyChangedEvent) override;
+	virtual EActorGridPlacement GetDefaultGridPlacement() const override { return EActorGridPlacement::AlwaysLoaded; }
+#endif // WITH_EDITOR
+	//~ End UObject/AActor Interface
+
+	FORCEINLINE bool IsDrawingEnabled() const { return bEnableDrawing; }
+	void UpdateDrawing() const;
+
+	FORCEINLINE bool IsRegistered() const { return bRegistered; }
+	void OnRegistered(const FZoneGraphDataHandle DataHandle);
+	void OnUnregistered();
+
+	// TODO: I wonder if the storage is unnecessary indirection?
+	FZoneGraphStorage& GetStorageMutable() { return ZoneStorage; }
+	const FZoneGraphStorage& GetStorage() const { return ZoneStorage; }
+	FCriticalSection& GetStorageLock() const { return ZoneStorageLock; }
+
+	FBox GetBounds() const;
+
+	/** @return Combined hash of all ZoneShapes that were used to build the data. */
+	uint32 GetCombinedShapeHash() const { return CombinedShapeHash; }
+
+	/** Sets Combined hash of all ZoneShapes that were used to build the data. */
+	void SetCombinedShapeHash(const uint32 Hash) { CombinedShapeHash = Hash; }
+
+protected:
+	bool RegisterWithSubsystem();
+	bool UnregisterWithSubsystem();
+
+	bool bRegistered;
+
+	/** if set to true then this zone graph data will be drawing itself when requested as part of "show navigation" */
+	UPROPERTY(Transient, EditAnywhere, Category = Display)
+	bool bEnableDrawing;
+
+	UPROPERTY(transient, duplicatetransient)
+	UZoneGraphRenderingComponent* RenderingComp;
+
+	UPROPERTY()
+	FZoneGraphStorage ZoneStorage;
+
+	/** Critical section to prevent rendering of the zone graph storage data while it's getting rebuilt */
+	mutable FCriticalSection ZoneStorageLock;
+
+	/** Combined hash of all ZoneShapes that were used to build the data. */
+	UPROPERTY()
+	uint32 CombinedShapeHash = 0;
+};
diff --git a/Engine/Plugins/Runtime/ZoneGraph/Source/ZoneGraph/Public/ZoneGraphRenderingComponent.h b/Engine/Plugins/Runtime/ZoneGraph/Source/ZoneGraph/Public/ZoneGraphRenderingComponent.h
new file mode 100644
--- /dev/null
+++ b/Engine/Plugins/Runtime/ZoneGraph/Source/ZoneGraph/Public/ZoneGraphRenderingComponent.h
@@ -0,0 +1,94 @@
+// Copyright Epic Games, Inc. All Rights Reserved.
+
+#pragma once
+
+#include "CoreMinimal.h"
+#include "ZoneGraphTypes.h"
+#include "Components/PrimitiveComponent.h"
+#include "DebugRenderSceneProxy.h"
+#include "ZoneGraphRenderingComponent.generated.h"
+
+class UZoneGraphRenderingComponent;
+class AZoneGraphData;
+
+// exported to API for GameplayDebugger module
+class ZONEGRAPH_API FZoneGraphSceneProxy : public FDebugRenderSceneProxy
+{
+public:
+	virtual SIZE_T GetTypeHash() const override;
+
+	FZoneGraphSceneProxy(const UPrimitiveComponent& InComponent, const AZoneGraphData& ZoneGraph);
+	virtual ~FZoneGraphSceneProxy();
+
+	virtual void GetDynamicMeshElements(const TArray<const FSceneView*>& Views, const FSceneViewFamily& ViewFamily, uint32 VisibilityMap, FMeshElementCollector& Collector) const override;
+
+	struct FZoneVisibility
+	{
+		bool bVisible = false;
+		bool bDetailsVisible = false;
+		float Alpha = 1.0f;
+	};
+	struct FDrawDistances
+	{
+		float MinDrawDistanceSqr = 0.f;
+		float MaxDrawDistanceSqr = FLT_MAX;
+		float FadeDrawDistanceSqr = FLT_MAX;
+		float DetailDrawDistanceSqr = FLT_MAX;
+	};
+	static FDrawDistances GetDrawDistances(const float MinDrawDistance, const float MaxDrawDistance);
+	static FZoneVisibility CalculateZoneVisibility(const FDrawDistances& Distances, const FVector Origin, const FVector Position);
+	static bool ShouldRenderZoneGraph(const FSceneView& View);
+
+protected:
+
+	virtual FPrimitiveViewRelevance GetViewRelevance(const FSceneView* View) const override;
+
+	virtual uint32 GetMemoryFootprint(void) const override { return sizeof(*this) + GetAllocatedSize(); }
+
+	uint32 GetAllocatedSize(void) const
+	{
+		return FDebugRenderSceneProxy::GetAllocatedSize();
+	}
+
+private:
+
+	TWeakObjectPtr<UZoneGraphRenderingComponent> WeakRenderingComponent;
+	bool bSkipDistanceCheck;
+};
+
+UCLASS(hidecategories = Object, editinlinenew)
+class ZONEGRAPH_API UZoneGraphRenderingComponent : public UPrimitiveComponent
+{
+	GENERATED_BODY()
+
+public:
+	UZoneGraphRenderingComponent(const FObjectInitializer& ObjectInitializer = FObjectInitializer::Get());
+
+	//~ Begin UPrimitiveComponent Interface
+	virtual FPrimitiveSceneProxy* CreateSceneProxy() override;
+	virtual void OnRegister()  override;
+	virtual void OnUnregister()  override;
+	//~ End UPrimitiveComponent Interface
+
+	//~ Begin USceneComponent Interface
+	virtual FBoxSphereBounds CalcBounds(const FTransform& LocalToWorld) const override;
+	//~ End USceneComponent Interface
+
+	void ForceUpdate() { bForceUpdate = true; }
+	bool IsForcingUpdate() const { return bForceUpdate; }
+
+	static bool IsNavigationShowFlagSet(const UWorld* World);
+
+protected:
+
+	void CheckDrawFlagTimerFunction();
+
+protected:
+#if !UE_BUILD_SHIPPING && !UE_BUILD_TEST
+	FDelegateHandle DebugTextDrawingDelegateHandle;
+#endif
+
+	bool bPreviousShowNavigation;
+	bool bForceUpdate;
+	FTimerHandle TimerHandle;
+};
diff --git a/Engine/Plugins/Runtime/ZoneGraph/Source/ZoneGraph/Public/ZoneGraphSubsystem.h b/Engine/Plugins/Runtime/ZoneGraph/Source/ZoneGraph/Public/ZoneGraphSubsystem.h
new file mode 100644
--- /dev/null
+++ b/Engine/Plugins/Runtime/ZoneGraph/Source/ZoneGraph/Public/ZoneGraphSubsystem.h
@@ -0,0 +1,160 @@
+// Copyright Epic Games, Inc. All Rights Reserved.
+
+#pragma once
+
+#include "ZoneGraphTypes.h"
+#include "ZoneGraphData.h"
+#include "Subsystems/WorldSubsystem.h"
+#if WITH_EDITOR
+#include "ZoneGraphBuilder.h"
+#endif
+#include "ZoneGraphSubsystem.generated.h"
+
+class UWorld;
+class UZoneShapeComponent;
+
+
+// Struct representing registered ZoneGraph data in the subsystem.
+USTRUCT()
+struct FRegisteredZoneGraphData
+{
+	GENERATED_BODY()
+
+	void Reset(int32 InGeneration = 1)
+	{
+		ZoneGraphData = nullptr;
+		bInUse = false;
+		Generation = InGeneration;
+	}
+
+	UPROPERTY()
+	AZoneGraphData* ZoneGraphData = nullptr;
+
+	int32 Generation = 1;	// Starting at generation 1 so that 0 can be invalid.
+	bool bInUse = false;	// Extra bit indicating that the data is in meant to be in use. This tried to capture the case where ZoneGraphData might get nullified without notifying.
+};
+
+
+UCLASS()
+class ZONEGRAPH_API UZoneGraphSubsystem : public UTickableWorldSubsystem
+{
+	GENERATED_BODY()
+public:
+	UZoneGraphSubsystem();
+
+	FZoneGraphDataHandle RegisterZoneGraphData(AZoneGraphData& InZoneGraphData);
+	void UnregisterZoneGraphData(AZoneGraphData& InZoneGraphData);
+	TConstArrayView<FRegisteredZoneGraphData> GetRegisteredZoneGraphData() const { return RegisteredZoneGraphData; }
+
+#if WITH_EDITOR
+	FZoneGraphBuilder& GetBuilder() { return Builder; }
+#endif
+
+	// Queries
+
+	// Returns Zone Graph data associated with specified handle.
+	const AZoneGraphData* GetZoneGraphData(const FZoneGraphDataHandle DataHandle) const;
+
+	// Returns Zone Graph data storage associated with specified handle, or nullptr if not found.
+	const FZoneGraphStorage* GetZoneGraphStorage(const FZoneGraphDataHandle DataHandle) const
+	{
+		if (int32(DataHandle.Generation) == RegisteredZoneGraphData[DataHandle.Index].Generation && int32(DataHandle.Index) < RegisteredZoneGraphData.Num())
+		{
+			if (const AZoneGraphData* Data = RegisteredZoneGraphData[DataHandle.Index].ZoneGraphData)
+			{
+				return &Data->GetStorage();
+			}
+		}
+		return nullptr;
+	}
+
+	// Find nearest lane that touches the query bounds. Finds results from all registered ZoneGraph data.
+	bool FindNearestLane(const FBox& QueryBounds, const FZoneGraphTagFilter TagFilter, FZoneGraphLaneLocation& OutLaneLocation, float& OutDistanceSqr) const;
+
+	// Find overlapping lanes that touches the query bounds. Finds results from all registered ZoneGraph data.
+	bool FindOverlappingLanes(const FBox& QueryBounds, const FZoneGraphTagFilter TagFilter, TArray<FZoneGraphLaneHandle>& OutLanes) const;
+
+	// Find sections of lanes fully overlapping (including lane width). Finds results from all registered ZoneGraph data.
+	bool FindLaneOverlaps(const FVector& Center, const float Radius, const FZoneGraphTagFilter TagFilter, TArray<FZoneGraphLaneSection>& OutLaneSections) const;
+	
+	// Moves LaneLocation along a lane.
+	bool AdvanceLaneLocation(const FZoneGraphLaneLocation& InLaneLocation, const float AdvanceDistance, FZoneGraphLaneLocation& OutLaneLocation) const;
+
+	// Returns point at a distance along a specific lane.
+	bool CalculateLocationAlongLane(const FZoneGraphLaneHandle LaneHandle, const float Distance, FZoneGraphLaneLocation& OutLaneLocation) const;
+
+	// Find nearest location on a specific lane.
+	bool FindNearestLocationOnLane(const FZoneGraphLaneHandle LaneHandle, const FBox& Bounds, FZoneGraphLaneLocation& OutLaneLocation, float& OutDistanceSqr) const;
+
+	// Find nearest location on a specific lane.
+	bool FindNearestLocationOnLane(const FZoneGraphLaneHandle LaneHandle, const FVector& Center, const float Range, FZoneGraphLaneLocation& OutLaneLocation, float& OutDistanceSqr) const;
+
+	// returns true if lane handle is valid.
+	bool IsLaneValid(const FZoneGraphLaneHandle LaneHandle) const;
+
+	// Returns the length of a specific lane.
+	bool GetLaneLength(const FZoneGraphLaneHandle LaneHandle, float& OutLength) const;
+
+	// Returns the width of a specific lane.
+	bool GetLaneWidth(const FZoneGraphLaneHandle LaneHandle, float& OutWidth) const;
+
+	// Returns the Tags of a specific lane.
+	bool GetLaneTags(const FZoneGraphLaneHandle LaneHandle, FZoneGraphTagMask& OutTags) const;
+
+	// Returns all links to connected lanes of a specific lane.
+	bool GetLinkedLanes(const FZoneGraphLaneHandle LaneHandle, const EZoneLaneLinkType Types, const EZoneLaneLinkFlags IncludeFlags, const EZoneLaneLinkFlags ExcludeFlags, TArray<FZoneGraphLinkedLane>& OutLinkedLanes) const;
+
+	// Returns handle to first linked lane matching the connection type and flags.
+	bool GetFirstLinkedLane(const FZoneGraphLaneHandle LaneHandle, const EZoneLaneLinkType Types, const EZoneLaneLinkFlags IncludeFlags, const EZoneLaneLinkFlags ExcludeFlags, FZoneGraphLinkedLane& OutLinkedLane) const;
+
+	// Returns bounds of all ZoneGraph data.
+	FBox GetCombinedBounds() const;
+	
+	// Tags
+	
+	// Returns tag based on name.
+	FZoneGraphTag GetTagByName(FName TagName) const;
+
+	// Returns the name of a specific tag.
+	FName GetTagName(FZoneGraphTag Tag) const;
+
+	// Returns info about a specific tag.
+	const FZoneGraphTagInfo* GetTagInfo(FZoneGraphTag Tag) const;
+
+	// Returns all tag infos
+	TConstArrayView<FZoneGraphTagInfo> GetTagInfos() const;
+
+protected:
+	virtual void Initialize(FSubsystemCollectionBase& Collection) override;
+	virtual void PostInitialize() override;
+	virtual void Deinitialize() override;
+	virtual void Tick(float DeltaTime) override;
+	virtual TStatId GetStatId() const override;
+
+	void RemoveRegisteredDataItem(const int32 Index);
+	void UnregisterStaleZoneGraphDataInstances();
+	void RegisterZoneGraphDataInstances();
+#if WITH_EDITOR
+	void OnActorMoved(AActor* Actor);
+	void OnRequestRebuild();
+	void SpawnMissingZoneGraphData();
+
+	/** Rebuilds the graph.  
+	 * @param bForceRebuild Settings this flag will force rebuild even if the data may be up to date. */
+	void RebuildGraph(const bool bForceRebuild = false);
+#endif // WITH_EDITOR
+
+	FCriticalSection DataRegistrationSection;
+
+	UPROPERTY()
+	TArray<FRegisteredZoneGraphData> RegisteredZoneGraphData;
+	TArray<int32> ZoneGraphDataFreeList;
+
+	bool bInitialized;
+
+#if WITH_EDITOR
+	FDelegateHandle OnActorMovedHandle;
+	FDelegateHandle OnRequestRebuildHandle;
+	FZoneGraphBuilder Builder;
+#endif
+};
diff --git a/Engine/Plugins/Runtime/ZoneGraph/Source/ZoneGraph/Public/ZoneShapeActor.h b/Engine/Plugins/Runtime/ZoneGraph/Source/ZoneGraph/Public/ZoneShapeActor.h
new file mode 100644
--- /dev/null
+++ b/Engine/Plugins/Runtime/ZoneGraph/Source/ZoneGraph/Public/ZoneShapeActor.h
@@ -0,0 +1,30 @@
+// Copyright Epic Games, Inc. All Rights Reserved.
+
+#pragma once
+
+#include "GameFramework/Actor.h"
+#include "ZoneShapeActor.generated.h"
+
+class UZoneShapeComponent;
+
+/** Zone Shape actor for standalone zone markup. */
+UCLASS(hidecategories = (Input))
+class ZONEGRAPH_API AZoneShape : public AActor
+{
+	GENERATED_BODY()
+public:
+	AZoneShape(const FObjectInitializer& ObjectInitializer = FObjectInitializer::Get());
+
+	virtual bool IsEditorOnly() const override { return true; }
+
+	const UZoneShapeComponent* GetShape() const { return ShapeComponent; }
+
+#if WITH_EDITOR
+	virtual EActorGridPlacement GetDefaultGridPlacement() const override { return EActorGridPlacement::AlwaysLoaded; }
+#endif
+
+protected:
+
+	UPROPERTY(Category = Zone, VisibleAnywhere, BlueprintReadOnly, meta = (AllowPrivateAccess = "true"))
+	UZoneShapeComponent* ShapeComponent;
+};
diff --git a/Engine/Plugins/Runtime/ZoneGraph/Source/ZoneGraphDebug/ZoneGraphDebug.Build.cs b/Engine/Plugins/Runtime/ZoneGraph/Source/ZoneGraphDebug/ZoneGraphDebug.Build.cs
new file mode 100644
--- /dev/null
+++ b/Engine/Plugins/Runtime/ZoneGraph/Source/ZoneGraphDebug/ZoneGraphDebug.Build.cs
@@ -0,0 +1,52 @@
+// Copyright Epic Games, Inc. All Rights Reserved.
+
+namespace UnrealBuildTool.Rules
+{
+	public class ZoneGraphDebug : ModuleRules
+	{
+		public ZoneGraphDebug(ReadOnlyTargetRules Target) : base(Target)
+		{
+			PCHUsage = ModuleRules.PCHUsageMode.UseExplicitOrSharedPCHs;
+
+			PublicIncludePaths.AddRange(
+			new string[] {
+				"Runtime/AIModule/Public",
+				ModuleDirectory + "/Public",
+			}
+			);
+
+			PrivateIncludePaths.AddRange(
+			new string[] {
+			}
+			);
+
+			PublicDependencyModuleNames.AddRange(
+			new string[] {
+				"Core",
+				"CoreUObject",
+				"Engine",
+				"GameplayTags",
+				"GameplayTasks",
+				"AIModule",
+				"ZoneGraph"
+			}
+			);
+
+			if (Target.bBuildEditor == true)
+			{
+				PrivateDependencyModuleNames.Add("EditorFramework");
+				PrivateDependencyModuleNames.Add("UnrealEd");
+			}
+
+			if (Target.bBuildDeveloperTools || (Target.Configuration != UnrealTargetConfiguration.Shipping && Target.Configuration != UnrealTargetConfiguration.Test))
+			{
+				PrivateDependencyModuleNames.Add("GameplayDebugger");
+				PublicDefinitions.Add("WITH_GAMEPLAY_DEBUGGER=1");
+			}
+			else
+			{
+				PublicDefinitions.Add("WITH_GAMEPLAY_DEBUGGER=0");
+			}
+		}
+	}
+}
diff --git a/Engine/Plugins/Runtime/ZoneGraph/Source/ZoneGraphEditor/Private/ZoneGraphEditorStyle.cpp b/Engine/Plugins/Runtime/ZoneGraph/Source/ZoneGraphEditor/Private/ZoneGraphEditorStyle.cpp
new file mode 100644
--- /dev/null
+++ b/Engine/Plugins/Runtime/ZoneGraph/Source/ZoneGraphEditor/Private/ZoneGraphEditorStyle.cpp
@@ -0,0 +1,60 @@
+﻿// Copyright Epic Games, Inc. All Rights Reserved.
+
+#include "ZoneGraphEditorStyle.h"
+#include "Styling/SlateStyleRegistry.h"
+#include "Styling/CoreStyle.h"
+#include "Styling/SlateTypes.h"
+#include "EditorStyleSet.h"
+#include "Interfaces/IPluginManager.h"
+
+#define IMAGE_BRUSH( RelativePath, ... ) FSlateImageBrush( StyleSet->RootToContentDir( RelativePath, TEXT( ".png" ) ), __VA_ARGS__ )
+#define IMAGE_PLUGIN_BRUSH( RelativePath, ... ) FSlateImageBrush( FMeshEditorStyle::InContent( RelativePath, ".png" ), __VA_ARGS__ )
+#define BOX_BRUSH( RelativePath, ... ) FSlateBoxBrush( StyleSet->RootToContentDir( RelativePath, TEXT( ".png" ) ), __VA_ARGS__ )
+#define TTF_CORE_FONT( RelativePath, ... ) FSlateFontInfo( StyleSet->RootToCoreContentDir( RelativePath, TEXT( ".ttf" ) ), __VA_ARGS__ )
+
+#define DEFAULT_FONT(...) FCoreStyle::GetDefaultFontStyle(__VA_ARGS__)
+
+TSharedPtr<FSlateStyleSet> FZoneGraphEditorStyle::StyleSet = nullptr;
+
+FString FZoneGraphEditorStyle::InContent(const FString& RelativePath, const ANSICHAR* Extension)
+{
+	static FString ContentDir = IPluginManager::Get().FindPlugin(TEXT("ZoneGraphEditor"))->GetContentDir() / TEXT("Slate");
+	return (ContentDir / RelativePath) + Extension;
+}
+
+void FZoneGraphEditorStyle::Initialize()
+{
+	if (StyleSet.IsValid())
+	{
+		return;
+	}
+
+	StyleSet = MakeShared<FSlateStyleSet>(GetStyleSetName());
+
+	StyleSet->SetContentRoot(FPaths::EngineContentDir() / TEXT("Editor/Slate"));
+	StyleSet->SetCoreContentRoot(FPaths::EngineContentDir() / TEXT("Slate"));
+
+	const FTextBlockStyle& NormalText = FEditorStyle::Get().GetWidgetStyle<FTextBlockStyle>("NormalText");
+
+	StyleSet->Set("ZoneGraph.Tag.Label", FTextBlockStyle(NormalText).SetFont(DEFAULT_FONT("Bold", 7)));
+
+	FSlateStyleRegistry::RegisterSlateStyle(*StyleSet.Get());
+}
+
+
+void FZoneGraphEditorStyle::Shutdown()
+{
+	if (StyleSet.IsValid())
+	{
+		FSlateStyleRegistry::UnRegisterSlateStyle(*StyleSet.Get());
+		ensure(StyleSet.IsUnique());
+		StyleSet.Reset();
+	}
+}
+
+
+FName FZoneGraphEditorStyle::GetStyleSetName()
+{
+	static FName StyleName("ZoneGraphEditorStyle");
+	return StyleName;
+}
diff --git a/Engine/Plugins/Runtime/ZoneGraph/Source/ZoneGraphEditor/Private/ZoneShapeComponentVisualizer.h b/Engine/Plugins/Runtime/ZoneGraph/Source/ZoneGraphEditor/Private/ZoneShapeComponentVisualizer.h
new file mode 100644
--- /dev/null
+++ b/Engine/Plugins/Runtime/ZoneGraph/Source/ZoneGraphEditor/Private/ZoneShapeComponentVisualizer.h
@@ -0,0 +1,273 @@
+// Copyright Epic Games, Inc. All Rights Reserved.
+
+#pragma once
+
+#include "CoreMinimal.h"
+#include "InputCoreTypes.h"
+#include "HitProxies.h"
+#include "ComponentVisualizer.h"
+#include "ZoneShapeComponent.h"
+#include "ZoneShapeComponentVisualizer.generated.h"
+
+class AActor;
+class FEditorViewportClient;
+class FMenuBuilder;
+class FPrimitiveDrawInterface;
+class FSceneView;
+class FUICommandList;
+class FViewport;
+class SWidget;
+struct FViewportClick;
+struct FConvexVolume;
+
+UENUM()
+enum class FZoneShapeControlPointType : uint8
+{
+	None,
+	In,
+	Out,
+};
+
+/** Selection state data that will be captured by scoped transactions.*/
+UCLASS(Transient)
+class ZONEGRAPHEDITOR_API UZoneShapeComponentVisualizerSelectionState : public UObject
+{
+	GENERATED_BODY()
+public:
+	FComponentPropertyPath GetShapePropertyPath() const { return ShapePropertyPath; }
+	void SetShapePropertyPath(const FComponentPropertyPath& InShapePropertyPath) { ShapePropertyPath = InShapePropertyPath; }
+	const TSet<int32>& GetSelectedPoints() const { return SelectedPoints; }
+	TSet<int32>& ModifySelectedPoints() { return SelectedPoints; }
+	int32 GetLastPointIndexSelected() const { return LastPointIndexSelected; }
+	void SetLastPointIndexSelected(const int32 InLastPointIndexSelected) { LastPointIndexSelected = InLastPointIndexSelected; }
+	int32 GetSelectedSegmentIndex() const { return SelectedSegmentIndex; }
+	void SetSelectedSegmentIndex(const int32 InSelectedSegmentIndex) { SelectedSegmentIndex = InSelectedSegmentIndex; }
+	FVector GetSelectedSegmentPoint() const { return SelectedSegmentPoint; }
+	void SetSelectedSegmentPoint(const FVector& InSelectedSegmentPoint) { SelectedSegmentPoint = InSelectedSegmentPoint; }
+	float GetSelectedSegmentT() const { return SelectedSegmentT; }
+	void SetSelectedSegmentT(const float InSelectedSegmentT) { SelectedSegmentT = InSelectedSegmentT; }
+	int32 GetSelectedControlPoint() const { return SelectedControlPoint; }
+	void SetSelectedControlPoint(const int32 InSelectedControlPoint) { SelectedControlPoint = InSelectedControlPoint; }
+	FZoneShapeControlPointType GetSelectedControlPointType() const { return SelectedControlPointType; }
+	void SetSelectedControlPointType(const FZoneShapeControlPointType InSelectedControlPointType) { SelectedControlPointType = InSelectedControlPointType; }
+
+protected:
+	/** Property path from the parent actor to the component */
+	UPROPERTY()
+	FComponentPropertyPath ShapePropertyPath;
+	/** Index of keys we have selected */
+	UPROPERTY()
+	TSet<int32> SelectedPoints;
+	/** Index of the last key we selected */
+	UPROPERTY()
+	int32 LastPointIndexSelected = 0;
+	/** Index of segment we have selected */
+	UPROPERTY()
+	int32 SelectedSegmentIndex = 0;
+	/** Position on selected segment */
+	UPROPERTY()
+	FVector SelectedSegmentPoint = FVector::ZeroVector;
+	/** Interpolation value along the selected segment */
+	UPROPERTY()
+	float SelectedSegmentT = 0.0f;
+	/** Index of tangent handle we have selected */
+	UPROPERTY()
+	int32 SelectedControlPoint = 0;
+	/** The type of the selected tangent handle */
+	UPROPERTY()
+	FZoneShapeControlPointType SelectedControlPointType = FZoneShapeControlPointType::None;
+};
+
+/** Base class for clickable shape editing proxies */
+struct ZONEGRAPHEDITOR_API HZoneShapeVisProxy : public HComponentVisProxy
+{
+	DECLARE_HIT_PROXY();
+
+	HZoneShapeVisProxy(const UActorComponent* InComponent, EHitProxyPriority InPriority = HPP_Wireframe)
+		: HComponentVisProxy(InComponent, InPriority)
+	{}
+
+	virtual EMouseCursor::Type GetMouseCursor() override
+	{
+		return EMouseCursor::CardinalCross;
+	}
+};
+
+/** Proxy for a shape point */
+struct ZONEGRAPHEDITOR_API HZoneShapePointProxy : public HZoneShapeVisProxy
+{
+	DECLARE_HIT_PROXY();
+
+	HZoneShapePointProxy(const UActorComponent* InComponent, int32 InPointIndex)
+		: HZoneShapeVisProxy(InComponent, HPP_Foreground)
+		, PointIndex(InPointIndex)
+	{}
+
+	virtual EMouseCursor::Type GetMouseCursor() override
+	{
+		return EMouseCursor::CardinalCross;
+	}
+
+	int32 PointIndex;
+};
+
+/** Proxy for a shape segment */
+struct ZONEGRAPHEDITOR_API HZoneShapeSegmentProxy : public HZoneShapeVisProxy
+{
+	DECLARE_HIT_PROXY();
+
+	HZoneShapeSegmentProxy(const UActorComponent* InComponent, int32 InSegmentIndex)
+		: HZoneShapeVisProxy(InComponent)
+		, SegmentIndex(InSegmentIndex)
+	{}
+
+	virtual EMouseCursor::Type GetMouseCursor() override
+	{
+		return EMouseCursor::CardinalCross;
+	}
+
+	int32 SegmentIndex;
+};
+
+/** Proxy for a control point handle */
+struct ZONEGRAPHEDITOR_API HZoneShapeControlPointProxy : public HZoneShapeVisProxy
+{
+	DECLARE_HIT_PROXY();
+
+	HZoneShapeControlPointProxy(const UActorComponent* InComponent, int32 InPointIndex, bool bInInControlPoint)
+		: HZoneShapeVisProxy(InComponent, HPP_Foreground)
+		, PointIndex(InPointIndex)
+		, bInControlPoint(bInInControlPoint)
+	{}
+
+	virtual EMouseCursor::Type GetMouseCursor() override
+	{
+		return EMouseCursor::CardinalCross;
+	}
+
+	int32 PointIndex;
+	bool bInControlPoint;
+};
+
+/** ZoneShapeComponent visualizer/edit functionality */
+class ZONEGRAPHEDITOR_API FZoneShapeComponentVisualizer : public FComponentVisualizer, public FGCObject
+{
+public:
+	FZoneShapeComponentVisualizer();
+	virtual ~FZoneShapeComponentVisualizer();
+
+	//~ Begin FComponentVisualizer Interface
+	virtual void OnRegister() override;
+	virtual void DrawVisualization(const UActorComponent* Component, const FSceneView* View, FPrimitiveDrawInterface* PDI) override;
+	virtual bool VisProxyHandleClick(FEditorViewportClient* InViewportClient, HComponentVisProxy* VisProxy, const FViewportClick& Click) override;
+	virtual void EndEditing() override;
+	virtual bool GetWidgetLocation(const FEditorViewportClient* ViewportClient, FVector& OutLocation) const override;
+	virtual bool GetCustomInputCoordinateSystem(const FEditorViewportClient* ViewportClient, FMatrix& OutMatrix) const override;
+	virtual bool HandleInputDelta(FEditorViewportClient* ViewportClient, FViewport* Viewport, FVector& DeltaTranslate, FRotator& DeltaRotate, FVector& DeltaScale) override;
+	virtual bool HandleInputKey(FEditorViewportClient* ViewportClient, FViewport* Viewport, FKey Key, EInputEvent Event) override;
+	/** Handle box select input */
+	virtual bool HandleBoxSelect(const FBox& InBox, FEditorViewportClient* InViewportClient, FViewport* InViewport) override;
+	/** Handle frustum select input */
+	virtual bool HandleFrustumSelect(const FConvexVolume& InFrustum, FEditorViewportClient* InViewportClient, FViewport* InViewport) override;
+	/** Return whether focus on selection should focus on bounding box defined by active visualizer */
+	virtual bool HasFocusOnSelectionBoundingBox(FBox& OutBoundingBox) override;
+	/** Pass snap input to active visualizer */
+	virtual bool HandleSnapTo(const bool bInAlign, const bool bInUseLineTrace, const bool bInUseBounds, const bool bInUsePivot, AActor* InDestination) override;
+	/** Get currently edited component, this is needed to reset the active visualizer after undo/redo */
+	virtual UActorComponent* GetEditedComponent() const override;
+	virtual TSharedPtr<SWidget> GenerateContextMenu() const override;
+	virtual bool IsVisualizingArchetype() const override;
+	//~ End FComponentVisualizer Interface
+
+	/** Get the shape component we are currently editing */
+	UZoneShapeComponent* GetEditedShapeComponent() const;
+
+	const TSet<int32>& GetSelectedPoints() const
+	{
+		check(SelectionState);
+		return SelectionState->GetSelectedPoints();
+	}
+
+protected:
+
+	/** Determine if any selected key index is out of range (perhaps because something external has modified the shape) */
+	bool IsAnySelectedPointIndexOutOfRange(const UZoneShapeComponent& Comp) const;
+
+	/** Whether a single point is currently selected */
+	bool IsSinglePointSelected() const;
+
+	/** Transforms selected control point by given translation */
+	bool TransformSelectedControlPoint(const FVector& DeltaTranslate);
+
+	/** Transforms selected points by given delta translation, rotation and scale. */
+	bool TransformSelectedPoints(const FEditorViewportClient* ViewportClient, const FVector& DeltaTranslate, const FRotator& DeltaRotate, const FVector& DeltaScale) const;
+
+	/** Update the key selection state of the visualizer */
+	void ChangeSelectionState(int32 Index, bool bIsCtrlHeld) const;
+
+	/** Alt-drag: duplicates the selected point */
+	bool DuplicatePointForAltDrag(const FVector& InDrag) const;
+
+	/** Split segment using given interpolation value, selects the new point */
+	void SplitSegment(const int32 InSegmentIndex, const float SegmentSplitT) const;
+
+	/** Duplicates selected points and selects them. */
+	void DuplicateSelectedPoints(const FVector& WorldOffset = FVector::ZeroVector, bool bInsertAfter = true) const;
+
+	/** Updates the component and selected properties if the component has changed */
+	const UZoneShapeComponent* UpdateSelectedShapeComponent(const HComponentVisProxy* VisProxy);
+
+	/** Returns the rotation of last selected point, or false if last selection is not valid. */
+	bool GetLastSelectedPointRotation(FQuat& OutRotation) const;
+
+	void OnDeletePoint() const;
+	bool CanDeletePoint() const;
+
+	/** Duplicates selected points in place */
+	void OnDuplicatePoint() const;
+	bool IsPointSelectionValid() const;
+
+	void OnAddPointToSegment() const;
+	bool CanAddPointToSegment() const;
+
+	void OnSetPointType(FZoneShapePointType Type) const;
+	bool IsPointTypeSet(FZoneShapePointType Type) const;
+
+	void OnSelectAllPoints() const;
+	bool CanSelectAllPoints() const;
+
+	void GenerateShapePointTypeSubMenu(FMenuBuilder& MenuBuilder) const;
+
+	// FGCObject interface
+	virtual void AddReferencedObjects(FReferenceCollector& Collector) override;
+	// End of FGCObject interface
+
+	/** Output log commands */
+	TSharedPtr<FUICommandList> ShapeComponentVisualizerActions;
+
+	/** Property to be notified when points change */
+	FProperty* ShapePointsProperty;
+
+	/** Current selection state */
+	UZoneShapeComponentVisualizerSelectionState* SelectionState;
+
+	/** Whether we currently allow duplication when dragging */
+	bool bAllowDuplication;
+
+	/** Alt-drag: Accumulates delayed drag offset. */
+	FVector DuplicateAccumulatedDrag;
+
+	/** True if the ControlPointPosition has been initialized */
+	bool bControlPointPositionCaptured;
+
+	/** Selected control pints adjusted position, allows the gizmo to move freely, while we constrain the control point. */
+	FVector ControlPointPosition;
+
+	/** True if cached rotation is set. */
+	bool bHasCachedRotation = false;
+
+	/** Rotation cached when mouse button is pressed. */
+	FQuat CachedRotation = FQuat::Identity;
+
+	bool bIsSelectingComponent = false;
+};
diff --git a/Engine/Plugins/Runtime/ZoneGraph/Source/ZoneGraphEditor/Public/ZoneGraphEditorModule.h b/Engine/Plugins/Runtime/ZoneGraph/Source/ZoneGraphEditor/Public/ZoneGraphEditorModule.h
new file mode 100644
--- /dev/null
+++ b/Engine/Plugins/Runtime/ZoneGraph/Source/ZoneGraphEditor/Public/ZoneGraphEditorModule.h
@@ -0,0 +1,36 @@
+// Copyright Epic Games, Inc. All Rights Reserved.
+
+#pragma  once
+
+#include "CoreMinimal.h"
+#include "Modules/ModuleManager.h"
+#include "AssetTypeCategories.h"
+#include "Toolkits/IToolkitHost.h"
+#include "Toolkits/AssetEditorToolkit.h"
+
+
+class IZoneGraphEditor;
+class FAssetTypeActions_Base;
+struct FGraphPanelNodeFactory;
+class FComponentVisualizer;
+
+/**
+* The public interface to this module
+*/
+class ZONEGRAPHEDITOR_API FZoneGraphEditorModule : public IModuleInterface//, public IHasMenuExtensibility, public IHasToolBarExtensibility
+{
+public:
+	// Begin IModuleInterface
+	virtual void StartupModule() override;
+	virtual void ShutdownModule() override;
+
+private:
+	void RegisterComponentVisualizer(FName ComponentClassName, TSharedPtr<FComponentVisualizer> Visualizer);
+	void RegisterMenus();
+
+	void OnBuildZoneGraph();
+
+	TArray<TSharedPtr<class FAssetTypeActions_Base>> ItemDataAssetTypeActions;
+	TArray<FName> RegisteredComponentClassNames;
+	TSharedPtr<class FUICommandList> PluginCommands;
+};
diff --git a/Engine/Plugins/Runtime/ZoneGraphAnnotations/Config/BaseZoneGraphAnnotations.ini b/Engine/Plugins/Runtime/ZoneGraphAnnotations/Config/BaseZoneGraphAnnotations.ini
new file mode 100644
--- /dev/null
+++ b/Engine/Plugins/Runtime/ZoneGraphAnnotations/Config/BaseZoneGraphAnnotations.ini
@@ -0,0 +1,11 @@
+﻿[CoreRedirects]
++PackageRedirects=(OldName="/Script/ZoneGraphBehaviors", NewName="/Script/ZoneGraphAnnotations")
++StructRedirects=(OldName="/Script/ZoneGraphBehaviors.ZoneGraphBehaviorEventBase", NewName="/Script/ZoneGraphAnnotations.ZoneGraphAnnotationEventBase")
++StructRedirects=(OldName="/Script/ZoneGraphBehaviors.ZoneGraphBehaviorTest", NewName="/Script/ZoneGraphAnnotations.ZoneGraphAnnotationTest")
++StructRedirects=(OldName="/Script/ZoneGraphBehaviors.ZoneGraphDangerArea", NewName="/Script/ZoneGraphAnnotations.ZoneGraphDisturbanceArea")
++ClassRedirects=(OldName="/Script/ZoneGraphBehaviors.ZoneGraphBehaviorTestingComponent", NewName="/Script/ZoneGraphAnnotations.ZoneGraphAnnotationTestingComponent")
++ClassRedirects=(OldName="/Script/ZoneGraphBehaviors.ZoneGraphBehaviorTestingActor", NewName="/Script/ZoneGraphAnnotations.ZoneGraphAnnotationTestingActor")
++ClassRedirects=(OldName="/Script/ZoneGraphBehaviors.ZoneGraphAnnotationSubsystem", NewName="/Script/ZoneGraphAnnotations.ZoneGraphAnnotationSubsystem")
++ClassRedirects=(OldName="/Script/ZoneGraphBehaviors.ZoneGraphFleeBehavior", NewName="/Script/ZoneGraphAnnotations.ZoneGraphDisturbanceAnnotation")
++ClassRedirects=(OldName="/Script/ZoneGraphBehaviors.ZoneGraphFleeBehaviorBPLibrary", NewName="/Script/ZoneGraphAnnotations.ZoneGraphDisturbanceAnnotationBPLibrary")
++PropertyRedirects=(OldName="ZoneGraphDisturbanceAnnotation.AnnotationTag",NewName="ZoneGraphDisturbanceAnnotation.DangerAnnotationTag")
\ No newline at end of file
diff --git a/Engine/Plugins/Runtime/ZoneGraphAnnotations/Source/ZoneGraphAnnotations/Public/Annotations/ZoneGraphDisturbanceAnnotationBPLibrary.h b/Engine/Plugins/Runtime/ZoneGraphAnnotations/Source/ZoneGraphAnnotations/Public/Annotations/ZoneGraphDisturbanceAnnotationBPLibrary.h
new file mode 100644
--- /dev/null
+++ b/Engine/Plugins/Runtime/ZoneGraphAnnotations/Source/ZoneGraphAnnotations/Public/Annotations/ZoneGraphDisturbanceAnnotationBPLibrary.h
@@ -0,0 +1,30 @@
+﻿// Copyright Epic Games, Inc. All Rights Reserved.
+
+#pragma once
+
+#include "CoreMinimal.h"
+#include "Kismet/BlueprintFunctionLibrary.h"
+#include "ZoneGraphTypes.h"
+#include "ZoneGraphDisturbanceAnnotationBPLibrary.generated.h"
+
+/**
+ *	Set of utilities for dealing with Disturbance Annotation.
+ */
+UCLASS()
+class ZONEGRAPHANNOTATIONS_API UZoneGraphDisturbanceAnnotationBPLibrary : public UBlueprintFunctionLibrary
+{
+	GENERATED_BODY()
+
+public:
+
+	/*
+	 * Triggers Danger event at specific location.
+	 * @param Instigator (optional) identifies this event coming from specific Instigator, only one danger will persist per instigator.
+	 * @param Position Position of the danger.
+	 * @param Radius Radius of the danger.
+	 * @param Duration Duration of the danger.
+	 */
+	UFUNCTION(BlueprintCallable, Category = "ZoneGraphAnnotations", meta = (WorldContext = "WorldContextObject"))
+	static void TriggerDanger(UObject* WorldContextObject, const AActor* Instigator, const FVector Position, const float Radius, const float Duration);
+	
+};
diff --git a/Engine/Plugins/Runtime/ZoneGraphAnnotations/Source/ZoneGraphAnnotations/Public/ZoneGraphAnnotationComponent.h b/Engine/Plugins/Runtime/ZoneGraphAnnotations/Source/ZoneGraphAnnotations/Public/ZoneGraphAnnotationComponent.h
new file mode 100644
--- /dev/null
+++ b/Engine/Plugins/Runtime/ZoneGraphAnnotations/Source/ZoneGraphAnnotations/Public/ZoneGraphAnnotationComponent.h
@@ -0,0 +1,100 @@
+﻿// Copyright Epic Games, Inc. All Rights Reserved.
+
+#pragma once
+
+#include "Components/PrimitiveComponent.h"
+#include "DebugRenderSceneProxy.h"
+#include "ZoneGraphTypes.h"
+#include "ZoneGraphAnnotationComponent.generated.h"
+
+class UZoneGraphAnnotationSubsystem;
+class UMassZoneGraphAnnotationComponent;
+class UCanvas;
+class AZoneGraphData;
+struct FInstancedStructStream;
+struct FZoneGraphAnnotationTagLookup;
+struct FZoneGraphAnnotationTagContainer;
+
+#if !UE_BUILD_SHIPPING && !UE_BUILD_TEST
+class ZONEGRAPHANNOTATIONS_API FZoneGraphAnnotationSceneProxy final : public FDebugRenderSceneProxy
+{
+public:
+	FZoneGraphAnnotationSceneProxy(const UPrimitiveComponent& InComponent, const EDrawType InDrawType = EDrawType::WireMesh);
+	
+	virtual SIZE_T GetTypeHash() const override;
+	virtual FPrimitiveViewRelevance GetViewRelevance(const FSceneView* View) const override;
+	virtual uint32 GetMemoryFootprint(void) const override;
+
+private:
+	uint32 ViewFlagIndex = 0;
+};
+#endif
+
+
+UCLASS(Abstract, HideCategories = (Physics, Collision, Lighting, Rendering, Mobile, HLOD, AssetUserData, Cooking, Navigation, Activation, Tags))
+class ZONEGRAPHANNOTATIONS_API UZoneGraphAnnotationComponent : public UPrimitiveComponent
+{
+	GENERATED_BODY()
+
+public:
+	UZoneGraphAnnotationComponent(const FObjectInitializer& ObjectInitializer = FObjectInitializer::Get());
+
+	/** Called during OnRegister(), or after all subsystems have been initialized. */
+	virtual void PostSubsystemsInitialized();
+	
+	/** Ticks the Annotation and changes the tags in the container when needed. */
+	virtual void TickAnnotation(const float DeltaTime, FZoneGraphAnnotationTagContainer& AnnotationTagContainer) {}
+
+	/** Called when new events are ready to be processed */
+	virtual void HandleEvents(TConstArrayView<const UScriptStruct*> AllEventStructs, const FInstancedStructStream& Events) {}
+
+	/** @return Tags applied by the Annotation, used to lookup Annotations from tags. */
+	virtual FZoneGraphTagMask GetAnnotationTags() const { return FZoneGraphTagMask::None; }
+
+	/** Called when new ZoneGraph data is added. */ 
+	virtual void PostZoneGraphDataAdded(const AZoneGraphData& ZoneGraphData) {}
+
+	/** Called when new ZoneGraph data is removed. */ 
+	virtual void PreZoneGraphDataRemoved(const AZoneGraphData& ZoneGraphData) {}
+	
+#if !UE_BUILD_SHIPPING && !UE_BUILD_TEST
+	/** Called when scene proxy is rebuilt. */
+	virtual void DebugDraw(FZoneGraphAnnotationSceneProxy* DebugProxy) {}
+
+	/** Called when it's time to draw to canvas. */
+	virtual void DebugDrawCanvas(UCanvas* Canvas, APlayerController*) {}
+#endif
+
+protected:
+
+	void OnPostZoneGraphDataAdded(const AZoneGraphData* ZoneGraphData);
+	void OnPreZoneGraphDataRemoved(const AZoneGraphData* ZoneGraphData);
+	void OnPostWorldInit(UWorld* World, const UWorld::InitializationValues);
+
+#if WITH_EDITOR
+	virtual void PostEditChangeChainProperty(FPropertyChangedChainEvent& PropertyChangedEvent) override;
+#endif // WITH_EDITOR
+
+	//~ Begin UPrimitiveComponent Interface
+	virtual void OnRegister()  override;
+	virtual void OnUnregister()  override;
+	virtual void DestroyRenderState_Concurrent() override;
+	virtual FPrimitiveSceneProxy* CreateSceneProxy() override;
+	//~ End UPrimitiveComponent Interface
+
+	//~ Begin USceneComponent Interface
+	virtual FBoxSphereBounds CalcBounds(const FTransform& LocalToWorld) const override;
+	//~ End USceneComponent Interface
+
+	FDelegateHandle OnPostZoneGraphDataAddedHandle;
+	FDelegateHandle OnPreZoneGraphDataRemovedHandle;
+	FDelegateHandle OnPostWorldInitDelegateHandle;
+
+#if !UE_BUILD_SHIPPING && !UE_BUILD_TEST
+	FDebugDrawDelegateHelper DebugDrawDelegateHelper;
+	FDelegateHandle CanvasDebugDrawDelegateHandle;
+#endif
+	
+	UPROPERTY(EditAnywhere, Category = Debug)
+	bool bEnableDebugDrawing = false;
+};
diff --git a/Engine/Plugins/Runtime/ZoneGraphAnnotations/Source/ZoneGraphAnnotations/Public/ZoneGraphAnnotationSubsystem.h b/Engine/Plugins/Runtime/ZoneGraphAnnotations/Source/ZoneGraphAnnotations/Public/ZoneGraphAnnotationSubsystem.h
new file mode 100644
--- /dev/null
+++ b/Engine/Plugins/Runtime/ZoneGraphAnnotations/Source/ZoneGraphAnnotations/Public/ZoneGraphAnnotationSubsystem.h
@@ -0,0 +1,143 @@
+﻿// Copyright Epic Games, Inc. All Rights Reserved.
+
+#pragma once
+
+#include "InstancedStructStream.h"
+#include "ZoneGraphTypes.h"
+#include "ZoneGraphAnnotationTypes.h"
+#include "Misc/MTAccessDetector.h"
+#include "Subsystems/WorldSubsystem.h"
+
+#include "ZoneGraphAnnotationSubsystem.generated.h"
+
+class UZoneGraphAnnotationComponent;
+class AZoneGraphData;
+
+
+/**
+ * Struct holding combined tags for a specific ZoneGraphData.
+ */
+struct ZONEGRAPHANNOTATIONS_API FZoneGraphDataAnnotationTags
+{
+	TArray<FZoneGraphTagMask> LaneTags;	// Combined array of tags from all Annotations.
+	FZoneGraphDataHandle DataHandle;	// Handle of the data
+	bool bInUse = false;				// True, if this entry is in use.
+};
+
+/**
+ * Annotation tags per zone graph data
+ */
+struct ZONEGRAPHANNOTATIONS_API FZoneGraphAnnotationTagContainer
+{
+	TArrayView<FZoneGraphTagMask> GetMutableAnnotationTagsForData(const FZoneGraphDataHandle DataHandle)
+	{
+		check(DataAnnotationTags[DataHandle.Index].DataHandle == DataHandle);
+		return DataAnnotationTags[DataHandle.Index].LaneTags;
+	}
+
+	TArray<FZoneGraphDataAnnotationTags> DataAnnotationTags;
+
+ 	/** Mask combining all static tags used by any of the registered ZoneGraphData. */
+	FZoneGraphTagMask CombinedStaticTags;
+};
+
+
+// Struct representing registered ZoneGraph data in the subsystem.
+USTRUCT()
+struct FRegisteredZoneGraphAnnotation
+{
+	GENERATED_BODY()
+
+	void Reset()
+	{
+		AnnotationComponent = nullptr;
+		AnnotationTags = FZoneGraphTagMask::None;
+	}
+
+	UPROPERTY()
+	UZoneGraphAnnotationComponent* AnnotationComponent = nullptr;
+
+	FZoneGraphTagMask AnnotationTags = FZoneGraphTagMask::None;	// Combination of all registered Annotation tag masks.
+};
+
+
+/**
+* A subsystem managing Zonegraph Annotations.
+*/
+UCLASS()
+class ZONEGRAPHANNOTATIONS_API UZoneGraphAnnotationSubsystem : public UTickableWorldSubsystem
+{
+	GENERATED_BODY()
+	
+public:
+	UZoneGraphAnnotationSubsystem();
+
+	/** Registers Annotation component */
+	void RegisterAnnotationComponent(UZoneGraphAnnotationComponent& Component);
+	
+	/** Unregisters Annotation component */
+	void UnregisterAnnotationComponent(UZoneGraphAnnotationComponent& Component);
+
+	/** Sends an event to the Annotations. */
+	template <typename T>
+	typename TEnableIf<TIsDerivedFrom<T, FZoneGraphAnnotationEventBase>::IsDerived, void>::Type SendEvent(const T& InRequest)
+	{
+		// This method is thread safe.
+		UE_MT_SCOPED_WRITE_ACCESS(EventsDetector);
+		Events[CurrentEventStream].Add(InRequest);
+	}
+
+	/** @return bitmask of Annotation tags at given lane */
+	FZoneGraphTagMask GetAnnotationTags(const FZoneGraphLaneHandle LaneHandle) const
+	{
+		check(AnnotationTagContainer.DataAnnotationTags.IsValidIndex(LaneHandle.DataHandle.Index));
+		const FZoneGraphDataAnnotationTags& AnnotationTags = AnnotationTagContainer.DataAnnotationTags[LaneHandle.DataHandle.Index];
+		return AnnotationTags.LaneTags[LaneHandle.Index];
+	}
+
+	/** @return First Annotation matching a bit in the bitmask */
+	UZoneGraphAnnotationComponent* GetFirstAnnotationForTag(const FZoneGraphTag AnnotationTag) const
+	{
+		return AnnotationTag.IsValid() ? TagToAnnotationLookup[AnnotationTag.Get()] : nullptr;
+	}
+
+	/** Signals the subsystem to re-register all tags. */
+#if WITH_EDITOR
+	void ReregisterTagsInEditor();
+#endif
+
+protected:
+
+	void PostZoneGraphDataAdded(const AZoneGraphData* ZoneGraphData);
+	void PreZoneGraphDataRemoved(const AZoneGraphData* ZoneGraphData);
+
+	void AddToAnnotationLookup(UZoneGraphAnnotationComponent& Annotation, const FZoneGraphTagMask AnnotationTags);
+	void RemoveFromAnnotationLookup(UZoneGraphAnnotationComponent& Annotation);
+	
+	virtual void Initialize(FSubsystemCollectionBase& Collection) override;
+	virtual void Deinitialize() override;
+	virtual void Tick(float DeltaTime) override;
+	virtual TStatId GetStatId() const override;
+
+	/** Array of registered components. */
+	UPROPERTY(Transient)
+	TArray<FRegisteredZoneGraphAnnotation> RegisteredComponents;
+
+	/** Stream of events to be processed, double buffered. */
+	UPROPERTY(Transient)
+	FInstancedStructStream Events[2];
+	UE_MT_DECLARE_RW_ACCESS_DETECTOR(EventsDetector);
+
+	/** Index of the current event stream. */
+	int32 CurrentEventStream = 0;
+
+	/** Lookup table from tag index to Annotation */
+	UPROPERTY(Transient)
+	TArray<UZoneGraphAnnotationComponent*> TagToAnnotationLookup;
+	
+	/** Combined tags for each ZoneGraphData. Each ZoneGraphData is indexed by it's data handle index, so there can be gaps in the array. */
+	FZoneGraphAnnotationTagContainer AnnotationTagContainer;
+
+	FDelegateHandle OnPostZoneGraphDataAddedHandle;
+	FDelegateHandle OnPreZoneGraphDataRemovedHandle;
+};
diff --git a/Engine/Plugins/Runtime/ZoneGraphAnnotations/Source/ZoneGraphAnnotations/Public/ZoneGraphAnnotationTestingActor.h b/Engine/Plugins/Runtime/ZoneGraphAnnotations/Source/ZoneGraphAnnotations/Public/ZoneGraphAnnotationTestingActor.h
new file mode 100644
--- /dev/null
+++ b/Engine/Plugins/Runtime/ZoneGraphAnnotations/Source/ZoneGraphAnnotations/Public/ZoneGraphAnnotationTestingActor.h
@@ -0,0 +1,101 @@
+﻿// Copyright Epic Games, Inc. All Rights Reserved.
+
+#pragma once
+
+#include "GameFramework/Actor.h"
+#include "Components/PrimitiveComponent.h"
+#include "DebugRenderSceneProxy.h"
+#include "ZoneGraphAnnotationTestingActor.generated.h"
+
+class UZoneGraphAnnotationTestingComponent;
+class FDebugRenderSceneProxy;
+
+/** Base class for ZoneGraph Annotation tests. */
+UCLASS(Abstract, EditInlineNew)
+class ZONEGRAPHANNOTATIONS_API UZoneGraphAnnotationTest : public UObject
+{
+	GENERATED_BODY()
+
+public:
+	virtual void Trigger() {}
+
+#if !UE_BUILD_SHIPPING && !UE_BUILD_TEST
+	virtual FBox CalcBounds(const FTransform& LocalToWorld) const { return FBox(ForceInit); };
+	virtual void DebugDraw(FDebugRenderSceneProxy* DebugProxy) {}
+	virtual void DebugDrawCanvas(UCanvas* Canvas, APlayerController*) {}
+#endif
+	
+	void SetOwner(UZoneGraphAnnotationTestingComponent* Owner) { OwnerComponent = Owner; OnOwnerSet(); }
+	const UZoneGraphAnnotationTestingComponent* GetOwner() const { return OwnerComponent; }
+
+protected:
+	virtual void OnOwnerSet() {}
+
+	UPROPERTY()
+	UZoneGraphAnnotationTestingComponent* OwnerComponent;
+};
+
+
+/** Debug component to test Mass ZoneGraph Annotations. Handles tests and rendering. */
+UCLASS(ClassGroup = Custom, HideCategories = (Physics, Collision, Lighting, Rendering, Mobile))
+class ZONEGRAPHANNOTATIONS_API UZoneGraphAnnotationTestingComponent : public UPrimitiveComponent
+{
+	GENERATED_BODY()
+public:
+	UZoneGraphAnnotationTestingComponent(const FObjectInitializer& ObjectInitializer = FObjectInitializer::Get());
+
+	virtual FBoxSphereBounds CalcBounds(const FTransform& LocalToWorld) const override;
+
+	void Trigger();
+
+protected:
+	
+#if WITH_EDITOR
+	virtual void PostEditChangeChainProperty(FPropertyChangedChainEvent& PropertyChangedEvent) override;
+#endif
+
+	//~ Begin UPrimitiveComponent Interface.
+	virtual void OnRegister() override;
+	virtual void OnUnregister() override;
+	virtual void DestroyRenderState_Concurrent() override;
+	virtual FPrimitiveSceneProxy* CreateSceneProxy() override;
+	//~ End UPrimitiveComponent Interface.
+
+#if !UE_BUILD_SHIPPING && !UE_BUILD_TEST
+	virtual void DebugDraw(FDebugRenderSceneProxy* DebugProxy);
+	virtual void DebugDrawCanvas(UCanvas* Canvas, APlayerController*);
+#endif
+	
+	UPROPERTY(EditAnywhere, Category = "Test", Instanced)
+	TArray<UZoneGraphAnnotationTest*> Tests;
+
+#if !UE_BUILD_SHIPPING && !UE_BUILD_TEST
+	FDebugDrawDelegateHelper DebugDrawDelegateHelper;
+	FDelegateHandle CanvasDebugDrawDelegateHandle;
+#endif
+};
+
+
+/** Debug actor to test Mass ZoneGraph Annotations. */
+UCLASS(HideCategories = (Actor, Input, Collision, Rendering, Replication, Partition, HLOD, Cooking))
+class ZONEGRAPHANNOTATIONS_API AZoneGraphAnnotationTestingActor : public AActor
+{
+	GENERATED_BODY()
+public:
+	AZoneGraphAnnotationTestingActor(const FObjectInitializer& ObjectInitializer = FObjectInitializer::Get());
+
+	/**
+	 * Simple trigger function to trigger something on the tests.
+	 * Ideally this would be part of each test, but it does not work there.
+	 */
+	UFUNCTION(BlueprintCallable, CallInEditor, Category = "Test")
+	void Trigger();
+
+protected:
+#if WITH_EDITOR
+	virtual void PostEditMove(bool bFinished) override;
+#endif
+
+	UPROPERTY(Category = Default, VisibleAnywhere, meta = (AllowPrivateAccess = "true"))
+	UZoneGraphAnnotationTestingComponent* TestingComp;
+};
diff --git a/Engine/Plugins/Runtime/ZoneGraphAnnotations/Source/ZoneGraphAnnotations/Public/ZoneGraphAnnotationTypes.h b/Engine/Plugins/Runtime/ZoneGraphAnnotations/Source/ZoneGraphAnnotations/Public/ZoneGraphAnnotationTypes.h
new file mode 100644
--- /dev/null
+++ b/Engine/Plugins/Runtime/ZoneGraphAnnotations/Source/ZoneGraphAnnotations/Public/ZoneGraphAnnotationTypes.h
@@ -0,0 +1,46 @@
+﻿// Copyright Epic Games, Inc. All Rights Reserved.
+
+#pragma once
+
+#include "CoreMinimal.h"
+#include "ZoneGraphAnnotationTypes.generated.h"
+
+ZONEGRAPHANNOTATIONS_API DECLARE_LOG_CATEGORY_EXTERN(LogZoneGraphAnnotations, Warning, All);
+
+/** Base structs for all events for UZoneGraphAnnotationSubsystem. */
+USTRUCT()
+struct ZONEGRAPHANNOTATIONS_API FZoneGraphAnnotationEventBase
+{
+	GENERATED_BODY()
+};
+
+class ZONEGRAPHANNOTATIONS_API FMassLaneObstacleID
+{
+public:
+	FMassLaneObstacleID() {}
+	
+	static const FMassLaneObstacleID InvalidID;
+	static FMassLaneObstacleID GetNextUniqueID()
+	{
+		check(NextUniqueID < MAX_uint64); // Ran out of FMassLaneObstacleID.
+		return FMassLaneObstacleID(NextUniqueID++);
+	}
+	
+	bool operator==(const FMassLaneObstacleID& Other) const
+	{
+		return Value == Other.Value;
+	}
+
+	uint64 GetValue() const { return Value; }
+
+private:
+	static uint64 NextUniqueID;
+	
+	FMassLaneObstacleID(uint64 ID) : Value(ID) {}
+	uint64 Value = MAX_uint64;
+};
+
+FORCEINLINE uint32 GetTypeHash(const FMassLaneObstacleID& Obs)
+{
+	return ::GetTypeHash(Obs.GetValue());
+}
