ZoneGraph: Fixed missing GC referenceer name

diff --git a/Engine/Plugins/Runtime/ZoneGraph/Source/ZoneGraphEditor/Private/ZoneShapeComponentVisualizer.h b/Engine/Plugins/Runtime/ZoneGraph/Source/ZoneGraphEditor/Private/ZoneShapeComponentVisualizer.h
--- a/Engine/Plugins/Runtime/ZoneGraph/Source/ZoneGraphEditor/Private/ZoneShapeComponentVisualizer.h
+++ b/Engine/Plugins/Runtime/ZoneGraph/Source/ZoneGraphEditor/Private/ZoneShapeComponentVisualizer.h
@@ -240,6 +240,10 @@ protected:
 
 	// FGCObject interface
 	virtual void AddReferencedObjects(FReferenceCollector& Collector) override;
+	virtual FString GetReferencerName() const override
+	{
+		return TEXT("FZoneShapeComponentVisualizer");
+	}
 	// End of FGCObject interface
 
 	/** Output log commands */
