ZoneGraph - Fix hand authored spline tangents in ZoneShapeUtilities.

diff --git a/Engine/Plugins/Runtime/ZoneGraph/Source/ZoneGraph/Private/ZoneShapeUtilities.cpp b/Engine/Plugins/Runtime/ZoneGraph/Source/ZoneGraph/Private/ZoneShapeUtilities.cpp
--- a/Engine/Plugins/Runtime/ZoneGraph/Source/ZoneGraph/Private/ZoneShapeUtilities.cpp
+++ b/Engine/Plugins/Runtime/ZoneGraph/Source/ZoneGraph/Private/ZoneShapeUtilities.cpp
@@ -952,14 +952,30 @@ void TessellateSplineShape(TConstArrayView<FZoneShapePoint> Points, const FZoneL
 		Lane.PointsEnd = OutZoneStorage.LanePoints.Num();
 
 		// Calculate per point forward.
-		const float ForwardScale = LaneDesc.Direction == EZoneLaneDirection::Forward ? 1.0f : -1.0f;
-		OutZoneStorage.LaneTangentVectors.Add(StartForward * ForwardScale);
+		if (LaneDesc.Direction == EZoneLaneDirection::Forward)
+		{
+			OutZoneStorage.LaneTangentVectors.Add(StartForward);
+		}
+		else
+		{
+			OutZoneStorage.LaneTangentVectors.Add(-EndForward);
+		}
+		
 		for (int32 PointIndex = Lane.PointsBegin + 1; PointIndex < Lane.PointsEnd - 1; PointIndex++)
 		{
 			const FVector WorldTangent = (OutZoneStorage.LanePoints[PointIndex + 1] - OutZoneStorage.LanePoints[PointIndex - 1]).GetSafeNormal();
 			OutZoneStorage.LaneTangentVectors.Add(WorldTangent);
 		}
-		OutZoneStorage.LaneTangentVectors.Add(EndForward * ForwardScale);
+
+		if (LaneDesc.Direction == EZoneLaneDirection::Forward)
+		{
+			OutZoneStorage.LaneTangentVectors.Add(EndForward);
+		}
+		else
+		{
+			OutZoneStorage.LaneTangentVectors.Add(-StartForward);
+		}
+
 
 		CurWidth += LaneDesc.Width;
 	}
