ZoneGraph fixes
- fixed crash accessing lane profile
- fixed ensure enum formatting to use UEnum::GetValueAsString()

diff --git a/Engine/Plugins/Runtime/ZoneGraph/Source/ZoneGraph/Private/ZoneGraphBuilder.cpp b/Engine/Plugins/Runtime/ZoneGraph/Source/ZoneGraph/Private/ZoneGraphBuilder.cpp
--- a/Engine/Plugins/Runtime/ZoneGraph/Source/ZoneGraph/Private/ZoneGraphBuilder.cpp
+++ b/Engine/Plugins/Runtime/ZoneGraph/Source/ZoneGraph/Private/ZoneGraphBuilder.cpp
@@ -288,13 +288,15 @@ void FZoneGraphBuilder::FindShapeConnections(const UZoneShapeComponent& SourceSh
 					&& FVector::DotProduct(SourceWorldNormal, -DestWorldNormal) > ConnectionSnapAngleCos)
 				{
 					// Check that the profile orientation matches before connecting.
-					const FZoneLaneProfile* LaneProfile = ZoneGraphSettings->GetLaneProfileByRef(SourceConnector.LaneProfile);
-					if (LaneProfile->IsSymmetrical() || SourceConnector.bReverseLaneProfile != DestConnector.bReverseLaneProfile)
+					if (const FZoneLaneProfile* LaneProfile = ZoneGraphSettings->GetLaneProfileByRef(SourceConnector.LaneProfile))
 					{
-						OutShapeConnections[i].ShapeComponent = DestShapeComp;
-						OutShapeConnections[i].ConnectorIndex = j;
-						bFound = true;
-						break;
+						if (LaneProfile->IsSymmetrical() || SourceConnector.bReverseLaneProfile != DestConnector.bReverseLaneProfile)
+						{
+							OutShapeConnections[i].ShapeComponent = DestShapeComp;
+							OutShapeConnections[i].ConnectorIndex = j;
+							bFound = true;
+							break;
+						}
 					}
 				}
 			}
diff --git a/Engine/Plugins/Runtime/ZoneGraph/Source/ZoneGraph/Public/ZoneGraphTypes.h b/Engine/Plugins/Runtime/ZoneGraph/Source/ZoneGraph/Public/ZoneGraphTypes.h
--- a/Engine/Plugins/Runtime/ZoneGraph/Source/ZoneGraph/Public/ZoneGraphTypes.h
+++ b/Engine/Plugins/Runtime/ZoneGraph/Source/ZoneGraph/Public/ZoneGraphTypes.h
@@ -164,7 +164,7 @@ struct ZONEGRAPH_API FZoneGraphTagMask
 		case EZoneLaneTagMaskComparison::Not:
 			return !ContainsAny(InTags);
 		default:
-			ensureMsgf(false, TEXT("Unhandled operand %s."), *StaticEnum<EZoneLaneTagMaskComparison>()->GetNameStringByValue((int32)Operand));
+			ensureMsgf(false, TEXT("Unhandled operand %s."), *UEnum::GetValueAsString(Operand));
 		}
 		return false;
 	}
