Clean up & sort #include paths in various private cpp files using work in progress FixIncludePath

#rnx

diff --git a/Engine/Plugins/Developer/Concert/ConcertSync/ConcertSyncClient/Source/ConcertSyncClient/Private/ConcertClientPresenceMode.cpp b/Engine/Plugins/Developer/Concert/ConcertSync/ConcertSyncClient/Source/ConcertSyncClient/Private/ConcertClientPresenceMode.cpp
--- a/Engine/Plugins/Developer/Concert/ConcertSync/ConcertSyncClient/Source/ConcertSyncClient/Private/ConcertClientPresenceMode.cpp
+++ b/Engine/Plugins/Developer/Concert/ConcertSync/ConcertSyncClient/Source/ConcertSyncClient/Private/ConcertClientPresenceMode.cpp
@@ -24,13 +24,13 @@
 #if WITH_EDITOR
 #include "Editor.h"
 #include "EditorWorldExtension.h"
-#include "LevelEditor.h"
 #include "IAssetViewport.h"
-#include "SLevelViewport.h"
 #include "IVREditorModule.h"
-#include "VREditorMode.h"
+#include "LevelEditor.h"
+#include "SLevelViewport.h"
+#include "Teleporter/VREditorTeleporter.h"
 #include "VREditorInteractor.h"
-#include "Editor/VREditor/Teleporter/VREditorTeleporter.h"
+#include "VREditorMode.h"
 #include "ViewportWorldInteraction.h"
 #endif
 
diff --git a/Engine/Plugins/Enterprise/DataprepEditor/Source/DataprepEditor/Private/AssetTypeActions_DataprepAssetInterface.cpp b/Engine/Plugins/Enterprise/DataprepEditor/Source/DataprepEditor/Private/AssetTypeActions_DataprepAssetInterface.cpp
--- a/Engine/Plugins/Enterprise/DataprepEditor/Source/DataprepEditor/Private/AssetTypeActions_DataprepAssetInterface.cpp
+++ b/Engine/Plugins/Enterprise/DataprepEditor/Source/DataprepEditor/Private/AssetTypeActions_DataprepAssetInterface.cpp
@@ -9,10 +9,10 @@
 #include "DataprepEditorModule.h"
 #include "DataprepFactories.h"
 
-#include "Developer/AssetTools/Public/IAssetTools.h"
-#include "Developer/AssetTools/Public/AssetToolsModule.h"
-#include "IContentBrowserSingleton.h"
+#include "AssetToolsModule.h"
 #include "ContentBrowserModule.h"
+#include "IAssetTools.h"
+#include "IContentBrowserSingleton.h"
 
 // UI
 #include "Framework/MultiBox/MultiBoxBuilder.h"
@@ -167,4 +167,4 @@ void FAssetTypeActions_DataprepAssetInterface::GetActions(const TArray<UObject*>
 	);
 }
 
-#undef LOCTEXT_NAMESPACE
\ No newline at end of file
+#undef LOCTEXT_NAMESPACE
diff --git a/Engine/Plugins/Enterprise/DataprepEditor/Source/DataprepEditor/Private/DataprepEditorModule.cpp b/Engine/Plugins/Enterprise/DataprepEditor/Source/DataprepEditor/Private/DataprepEditorModule.cpp
--- a/Engine/Plugins/Enterprise/DataprepEditor/Source/DataprepEditor/Private/DataprepEditorModule.cpp
+++ b/Engine/Plugins/Enterprise/DataprepEditor/Source/DataprepEditor/Private/DataprepEditorModule.cpp
@@ -13,8 +13,8 @@
 #include "Widgets/SDataprepEditorViewport.h"
 #include "Widgets/SDataprepProducersWidget.h"
 
-#include "Developer/AssetTools/Public/IAssetTools.h"
-#include "Developer/AssetTools/Public/AssetToolsModule.h"
+#include "AssetToolsModule.h"
+#include "IAssetTools.h"
 #include "Kismet2/KismetEditorUtilities.h"
 #include "Misc/PackageName.h"
 #include "Modules/ModuleManager.h"
diff --git a/Engine/Plugins/Enterprise/DatasmithContent/Source/DatasmithContentEditor/Private/DatasmithContentEditorModule.cpp b/Engine/Plugins/Enterprise/DatasmithContent/Source/DatasmithContentEditor/Private/DatasmithContentEditorModule.cpp
--- a/Engine/Plugins/Enterprise/DatasmithContent/Source/DatasmithContentEditor/Private/DatasmithContentEditorModule.cpp
+++ b/Engine/Plugins/Enterprise/DatasmithContent/Source/DatasmithContentEditor/Private/DatasmithContentEditorModule.cpp
@@ -9,11 +9,11 @@
 #include "DatasmithSceneActorDetailsPanel.h"
 #include "DatasmithSceneDetails.h"
 
+#include "AssetToolsModule.h"
 #include "DatasmithAssetImportData.h"
 #include "DatasmithScene.h"
-#include "Developer/AssetTools/Public/IAssetTools.h"
-#include "Developer/AssetTools/Public/AssetToolsModule.h"
 #include "Engine/Blueprint.h"
+#include "IAssetTools.h"
 #include "Modules/ModuleManager.h"
 #include "PropertyEditorModule.h"
 
diff --git a/Engine/Plugins/Experimental/ColorCorrectRegions/Source/ColorCorrectRegionsEditor/Private/ColorCorrectRegionsEditorModule.cpp b/Engine/Plugins/Experimental/ColorCorrectRegions/Source/ColorCorrectRegionsEditor/Private/ColorCorrectRegionsEditorModule.cpp
--- a/Engine/Plugins/Experimental/ColorCorrectRegions/Source/ColorCorrectRegionsEditor/Private/ColorCorrectRegionsEditorModule.cpp
+++ b/Engine/Plugins/Experimental/ColorCorrectRegions/Source/ColorCorrectRegionsEditor/Private/ColorCorrectRegionsEditorModule.cpp
@@ -1,12 +1,12 @@
 // Copyright Epic Games, Inc. All Rights Reserved.
 
 #include "ColorCorrectRegionsEditorModule.h"
-#include "Editor/PlacementMode/Public/IPlacementModeModule.h"
+#include "ActorFactories/ActorFactoryBlueprint.h"
 #include "ColorCorrectRegion.h"
-#include "ColorCorrectWindow.h"
-#include "ColorCorrectRegionsStyle.h"
 #include "ColorCorrectRegionCustomization.h"
-#include "ActorFactories/ActorFactoryBlueprint.h"
+#include "ColorCorrectRegionsStyle.h"
+#include "ColorCorrectWindow.h"
+#include "IPlacementModeModule.h"
 
 #define LOCTEXT_NAMESPACE "FColorCorrectRegionsModule"
 
diff --git a/Engine/Plugins/Media/BinkMedia/Source/BinkMediaPlayer/Private/BinkMediaPlayerModule.cpp b/Engine/Plugins/Media/BinkMedia/Source/BinkMediaPlayer/Private/BinkMediaPlayerModule.cpp
--- a/Engine/Plugins/Media/BinkMedia/Source/BinkMediaPlayer/Private/BinkMediaPlayerModule.cpp
+++ b/Engine/Plugins/Media/BinkMedia/Source/BinkMediaPlayer/Private/BinkMediaPlayerModule.cpp
@@ -18,22 +18,22 @@ TArray< FTexture2DRHIRef > BinkActiveTextureRefs;
 
 #if BINKPLUGIN_UE4_EDITOR
 class UFactory;
-#include "ISettingsModule.h"
-#include "PropertyEditorModule.h"
-#include "PropertyCustomizationHelpers.h"
+#include "DesktopPlatformModule.h"
 #include "DetailCategoryBuilder.h"
 #include "DetailLayoutBuilder.h"
-#include "IDetailPropertyRow.h"
-#include "DesktopPlatformModule.h"
+#include "Editor.h"
 #include "EditorDirectories.h"
-#include "SourceControlHelpers.h"
 #include "Framework/Notifications/NotificationManager.h"
-#include "Widgets/Notifications/SNotificationList.h"
-#include "Widgets/Input/SFilePathPicker.h"
-#include "Editor/MainFrame/Public/Interfaces/IMainFrameModule.h"
-#include "Editor/UnrealEd/Public/Editor.h"
-#include "IDetailCustomization.h"
 #include "IDetailChildrenBuilder.h"
+#include "IDetailCustomization.h"
+#include "IDetailPropertyRow.h"
+#include "ISettingsModule.h"
+#include "Interfaces/IMainFrameModule.h"
+#include "PropertyCustomizationHelpers.h"
+#include "PropertyEditorModule.h"
+#include "SourceControlHelpers.h"
+#include "Widgets/Input/SFilePathPicker.h"
+#include "Widgets/Notifications/SNotificationList.h"
 
 #include "BinkMediaPlayer.h"
 
diff --git a/Engine/Plugins/Media/BinkMedia/Source/BinkMediaPlayer/Private/binkpluginRHI.cpp b/Engine/Plugins/Media/BinkMedia/Source/BinkMediaPlayer/Private/binkpluginRHI.cpp
--- a/Engine/Plugins/Media/BinkMedia/Source/BinkMediaPlayer/Private/binkpluginRHI.cpp
+++ b/Engine/Plugins/Media/BinkMedia/Source/BinkMediaPlayer/Private/binkpluginRHI.cpp
@@ -8,17 +8,16 @@
 #include "binkplugin.h"
 
 #include "CoreMinimal.h"
-#include "UObject/ConstructorHelpers.h"
-#include "RHI.h"
+#include "MaterialShader.h"
+#include "Materials/Material.h"
 #include "Materials/MaterialInstanceDynamic.h"
 #include "Materials/MaterialInterface.h"
-#include "Materials/Material.h"
-#include "Runtime/RHI/Public/RHIStaticStates.h"
-#include "Runtime/RenderCore/Public/ShaderParameterUtils.h"
-#include "Runtime/RenderCore/Public/RenderResource.h"
-#include "Runtime/Renderer/Public/MaterialShader.h"
-#include "Runtime/RenderCore/Public/RenderGraphResources.h"
-#include "Runtime/RenderCore/Public/RenderGraphResources.h"
+#include "RHI.h"
+#include "RHIStaticStates.h"
+#include "RenderGraphResources.h"
+#include "RenderResource.h"
+#include "ShaderParameterUtils.h"
+#include "UObject/ConstructorHelpers.h"
 
 static BINKSHADERS * shaders;
 
diff --git a/Engine/Plugins/Media/BinkMedia/Source/BinkMediaPlayer/Private/binktexturesRHI.cpp b/Engine/Plugins/Media/BinkMedia/Source/BinkMediaPlayer/Private/binktexturesRHI.cpp
--- a/Engine/Plugins/Media/BinkMedia/Source/BinkMediaPlayer/Private/binktexturesRHI.cpp
+++ b/Engine/Plugins/Media/BinkMedia/Source/BinkMediaPlayer/Private/binktexturesRHI.cpp
@@ -6,20 +6,20 @@
 #include "../../BinkMediaPlayerSDK/include/binktiny.h"
 #include "../../BinkMediaPlayerSDK/include/binktextures.h"
 
+#include "BinkShaders.h"
 #include "CoreMinimal.h"
-#include "UObject/ConstructorHelpers.h"
-#include "RHI.h"
+#include "MaterialShader.h"
+#include "Materials/Material.h"
 #include "Materials/MaterialInstanceDynamic.h"
 #include "Materials/MaterialInterface.h"
-#include "Materials/Material.h"
-#include "Runtime/RHI/Public/RHIStaticStates.h"
-#include "Runtime/RHI/Public/PipelineStateCache.h"
-#include "Runtime/RenderCore/Public/ShaderParameterUtils.h"
-#include "Runtime/RenderCore/Public/RenderResource.h"
-#include "Runtime/Renderer/Public/MaterialShader.h"
-#include "Runtime/RenderCore/Public/RenderGraphBuilder.h"
-#include "Runtime/RenderCore/Public/RenderGraphResources.h"
-#include "Runtime/RenderCore/Public/BinkShaders.h"
+#include "PipelineStateCache.h"
+#include "RHI.h"
+#include "RHIStaticStates.h"
+#include "RenderGraphBuilder.h"
+#include "RenderGraphResources.h"
+#include "RenderResource.h"
+#include "ShaderParameterUtils.h"
+#include "UObject/ConstructorHelpers.h"
 
 
 extern FRHITexture2D *BinkRHIRenderTarget;
diff --git a/Engine/Plugins/Media/BinkMedia/Source/BinkMediaPlayerEditor/Private/BinkMediaPlayerEditorModule.cpp b/Engine/Plugins/Media/BinkMedia/Source/BinkMediaPlayerEditor/Private/BinkMediaPlayerEditorModule.cpp
--- a/Engine/Plugins/Media/BinkMedia/Source/BinkMediaPlayerEditor/Private/BinkMediaPlayerEditorModule.cpp
+++ b/Engine/Plugins/Media/BinkMedia/Source/BinkMediaPlayerEditor/Private/BinkMediaPlayerEditorModule.cpp
@@ -3,17 +3,17 @@
 
 
 #include "AssetTypeActions_Base.h"
-#include "BinkMediaPlayerEditorPrivate.h"
-#include "TextureEditor.h"
-#include "Editor/ContentBrowser/Public/ContentBrowserModule.h"
-#include "Modules/ModuleInterface.h"
 #include "BinkMediaPlayerCustomization.h"
-#include "Interfaces/ITextureEditorModule.h"
-#include "IContentBrowserSingleton.h"
-#include "Styling/SlateStyleRegistry.h"
-#include "Framework/MultiBox/MultiBoxBuilder.h"
+#include "BinkMediaPlayerEditorPrivate.h"
+#include "ContentBrowserModule.h"
 #include "Factories/MaterialFactoryNew.h"
+#include "Framework/MultiBox/MultiBoxBuilder.h"
+#include "IContentBrowserSingleton.h"
+#include "Interfaces/ITextureEditorModule.h"
 #include "Materials/Material.h"
+#include "Modules/ModuleInterface.h"
+#include "Styling/SlateStyleRegistry.h"
+#include "TextureEditor.h"
 #include "ThumbnailRendering/TextureThumbnailRenderer.h"
 #include "Widgets/SViewport.h"
 
diff --git a/Engine/Plugins/Runtime/HairStrands/Source/HairStrandsEditor/Private/GroomEditorViewportToolBar.cpp b/Engine/Plugins/Runtime/HairStrands/Source/HairStrandsEditor/Private/GroomEditorViewportToolBar.cpp
--- a/Engine/Plugins/Runtime/HairStrands/Source/HairStrandsEditor/Private/GroomEditorViewportToolBar.cpp
+++ b/Engine/Plugins/Runtime/HairStrands/Source/HairStrandsEditor/Private/GroomEditorViewportToolBar.cpp
@@ -1,20 +1,20 @@
 // Copyright Epic Games, Inc. All Rights Reserved.
 
 #include "GroomEditorViewportToolBar.h"
-#include "Widgets/SGroomEditorViewport.h"
-#include "Widgets/SBoxPanel.h"
-#include "Widgets/Text/STextBlock.h"
-#include "UObject/Package.h"
-#include "GroomComponent.h"
-#include "Styling/AppStyle.h"
-#include "Slate/SceneViewport.h"
+#include "AnalyticsEventAttribute.h"
 #include "ComponentReregisterContext.h"
-#include "Runtime/Analytics/Analytics/Public/AnalyticsEventAttribute.h"
-#include "Runtime/Analytics/Analytics/Public/Interfaces/IAnalyticsProvider.h"
 #include "EngineAnalytics.h"
-#include "Widgets/Docking/SDockTab.h"
-#include "Editor/UnrealEd/Public/SEditorViewportToolBarMenu.h"
+#include "GroomComponent.h"
 #include "GroomEditorCommands.h"
+#include "Interfaces/IAnalyticsProvider.h"
+#include "SEditorViewportToolBarMenu.h"
+#include "Slate/SceneViewport.h"
+#include "Styling/AppStyle.h"
+#include "UObject/Package.h"
+#include "Widgets/Docking/SDockTab.h"
+#include "Widgets/SBoxPanel.h"
+#include "Widgets/SGroomEditorViewport.h"
+#include "Widgets/Text/STextBlock.h"
 
 #define LOCTEXT_NAMESPACE "GroomEditorViewportToolbar"
 
diff --git a/Engine/Plugins/Runtime/MassGameplay/Source/MassMovementEditor/Private/MassMovementEditorModule.cpp b/Engine/Plugins/Runtime/MassGameplay/Source/MassMovementEditor/Private/MassMovementEditorModule.cpp
--- a/Engine/Plugins/Runtime/MassGameplay/Source/MassMovementEditor/Private/MassMovementEditorModule.cpp
+++ b/Engine/Plugins/Runtime/MassGameplay/Source/MassMovementEditor/Private/MassMovementEditorModule.cpp
@@ -1,11 +1,11 @@
 // Copyright Epic Games, Inc. All Rights Reserved.
 
 #include "MassMovementEditorModule.h"
+#include "AssetToolsModule.h"
+#include "MassMovementStyleRefDetails.h"
 #include "Modules/ModuleManager.h"
-#include "Developer/AssetTools/Public/AssetToolsModule.h"
 #include "PropertyEditorModule.h"
 #include "ToolMenus.h"
-#include "MassMovementStyleRefDetails.h"
 
 #define LOCTEXT_NAMESPACE "MassMovementEditor"
 
diff --git a/Engine/Plugins/Runtime/Steam/SteamAudio/Source/SteamAudioEditor/Private/BakeIndirectWindow.cpp b/Engine/Plugins/Runtime/Steam/SteamAudio/Source/SteamAudioEditor/Private/BakeIndirectWindow.cpp
--- a/Engine/Plugins/Runtime/Steam/SteamAudio/Source/SteamAudioEditor/Private/BakeIndirectWindow.cpp
+++ b/Engine/Plugins/Runtime/Steam/SteamAudio/Source/SteamAudioEditor/Private/BakeIndirectWindow.cpp
@@ -10,24 +10,24 @@
 #include "TickableNotification.h"
 #include "IndirectBaker.h"
 
-#include "Framework/Docking/TabManager.h"
+#include "Async/Async.h"
 #include "DetailLayoutBuilder.h"
+#include "EngineUtils.h"
+#include "Framework/Docking/TabManager.h"
+#include "Framework/Notifications/NotificationManager.h"
+#include "Kismet/GameplayStatics.h"
 #include "LevelEditor.h"
 #include "LevelEditorActions.h"
 #include "LevelEditorViewport.h"
-#include "EngineUtils.h"
 #include "Widgets/Docking/SDockTab.h"
+#include "Widgets/Input/SButton.h"
+#include "Widgets/Input/SComboBox.h"
+#include "Widgets/Input/SSpinBox.h"
 #include "Widgets/Layout/SExpandableArea.h"
 #include "Widgets/Layout/SGridPanel.h"
-#include "Widgets/Input/SSpinBox.h"
 #include "Widgets/Layout/SScrollBox.h"
-#include "Widgets/Input/SComboBox.h"
-#include "Widgets/Input/SButton.h"
+#include "Widgets/Notifications/SNotificationList.h"
 #include "Widgets/Views/SListView.h"
-#include "Async/Async.h"
-#include "Kismet/GameplayStatics.h"
-#include "Runtime/Slate/Public/Widgets/Notifications/SNotificationList.h"
-#include "Runtime/Slate/Public/Framework/Notifications/NotificationManager.h"
 
 namespace SteamAudio
 {
diff --git a/Engine/Plugins/Runtime/Steam/SteamVR/Source/SteamVR/Private/SteamVRHMD.cpp b/Engine/Plugins/Runtime/Steam/SteamVR/Source/SteamVR/Private/SteamVRHMD.cpp
--- a/Engine/Plugins/Runtime/Steam/SteamVR/Source/SteamVR/Private/SteamVRHMD.cpp
+++ b/Engine/Plugins/Runtime/Steam/SteamVR/Source/SteamVR/Private/SteamVRHMD.cpp
@@ -17,13 +17,13 @@
 #include "IHeadMountedDisplayVulkanExtensions.h"
 
 #include "EngineAnalytics.h"
-#include "Runtime/Analytics/Analytics/Public/Interfaces/IAnalyticsProvider.h"
+#include "Interfaces/IAnalyticsProvider.h"
 
 #include "Engine/Canvas.h"
 #include "CanvasItem.h"
 
 #if WITH_EDITOR
-#include "Editor/UnrealEd/Classes/Editor/EditorEngine.h"
+#include "Editor/EditorEngine.h"
 #endif
 
 #ifndef ARRAYSIZE
diff --git a/Engine/Plugins/Runtime/Steam/SteamVR/Source/SteamVREditor/Private/SteamVREditorStyle.cpp b/Engine/Plugins/Runtime/Steam/SteamVR/Source/SteamVREditor/Private/SteamVREditorStyle.cpp
--- a/Engine/Plugins/Runtime/Steam/SteamVR/Source/SteamVREditor/Private/SteamVREditorStyle.cpp
+++ b/Engine/Plugins/Runtime/Steam/SteamVR/Source/SteamVREditor/Private/SteamVREditorStyle.cpp
@@ -29,8 +29,8 @@ POSSIBILITY OF SUCH DAMAGE.
 */
 
 #include "SteamVREditorStyle.h"
+#include "Interfaces/IPluginManager.h"
 #include "SteamVREditor.h"
-#include "Runtime/Projects/Public/Interfaces/IPluginManager.h"
 
 TSharedPtr< FSlateStyleSet > FSteamVREditorStyle::StyleInstance = NULL;
 
diff --git a/Engine/Plugins/Runtime/Steam/SteamVR/Source/SteamVRInputDevice/Private/AnimNode_SteamVRInputAnimPose.cpp b/Engine/Plugins/Runtime/Steam/SteamVR/Source/SteamVRInputDevice/Private/AnimNode_SteamVRInputAnimPose.cpp
--- a/Engine/Plugins/Runtime/Steam/SteamVR/Source/SteamVRInputDevice/Private/AnimNode_SteamVRInputAnimPose.cpp
+++ b/Engine/Plugins/Runtime/Steam/SteamVR/Source/SteamVRInputDevice/Private/AnimNode_SteamVRInputAnimPose.cpp
@@ -29,9 +29,9 @@ POSSIBILITY OF SUCH DAMAGE.
 */
 
 #include "AnimNode_SteamVRInputAnimPose.h"
-#include "ISteamVRInputDeviceModule.h"
+#include "Animation/AnimInstanceProxy.h"
 #include "AnimationRuntime.h"
-#include "Runtime/Engine/Public/Animation/AnimInstanceProxy.h"
+#include "ISteamVRInputDeviceModule.h"
 #include "SteamVRInputDevice.h"
 #include "UE4HandSkeletonDefinition.h"
 
diff --git a/Engine/Plugins/Runtime/Steam/SteamVR/Source/SteamVRInputDevice/Private/AnimNode_SteamVRSetWristTransform.cpp b/Engine/Plugins/Runtime/Steam/SteamVR/Source/SteamVRInputDevice/Private/AnimNode_SteamVRSetWristTransform.cpp
--- a/Engine/Plugins/Runtime/Steam/SteamVR/Source/SteamVRInputDevice/Private/AnimNode_SteamVRSetWristTransform.cpp
+++ b/Engine/Plugins/Runtime/Steam/SteamVR/Source/SteamVRInputDevice/Private/AnimNode_SteamVRSetWristTransform.cpp
@@ -29,9 +29,9 @@ POSSIBILITY OF SUCH DAMAGE.
 */
 
 #include "AnimNode_SteamVRSetWristTransform.h"
-#include "ISteamVRInputDeviceModule.h"
+#include "Animation/AnimInstanceProxy.h"
 #include "AnimationRuntime.h"
-#include "Runtime/Engine/Public/Animation/AnimInstanceProxy.h"
+#include "ISteamVRInputDeviceModule.h"
 #include "SteamVRInputDevice.h"
 #include "UE4HandSkeletonDefinition.h"
 
diff --git a/Engine/Plugins/Runtime/Steam/SteamVR/Source/SteamVRInputDevice/Private/SteamVRInputDevice.cpp b/Engine/Plugins/Runtime/Steam/SteamVR/Source/SteamVRInputDevice/Private/SteamVRInputDevice.cpp
--- a/Engine/Plugins/Runtime/Steam/SteamVR/Source/SteamVRInputDevice/Private/SteamVRInputDevice.cpp
+++ b/Engine/Plugins/Runtime/Steam/SteamVR/Source/SteamVRInputDevice/Private/SteamVRInputDevice.cpp
@@ -29,23 +29,23 @@ POSSIBILITY OF SUCH DAMAGE.
 */
 
 #include "SteamVRInputDevice.h"
-#include "Runtime/ApplicationCore/Public/GenericPlatform/IInputInterface.h"
-#include "HAL/FileManagerGeneric.h"
-#include "Misc/FileHelper.h"
-#include "GameFramework/PlayerInput.h"
 #include "Engine/World.h"
-#include "Kismet/GameplayStatics.h"
-#include "openvr.h"
-#include "GameFramework/WorldSettings.h"
 #include "Features/IModularFeatures.h"
-#include "MotionControllerComponent.h"
+#include "GameFramework/PlayerInput.h"
+#include "GameFramework/WorldSettings.h"
+#include "GenericPlatform/GenericPlatformInputDeviceMapper.h"
+#include "GenericPlatform/IInputInterface.h"
+#include "HAL/FileManagerGeneric.h"
 #include "IMotionController.h"
-#include "SteamVRControllerKeys.h"
-#include "Runtime/HeadMountedDisplay/Public/IXRTrackingSystem.h"
-#include "SteamVRSkeletonDefinition.h"
+#include "IXRTrackingSystem.h"
+#include "Kismet/GameplayStatics.h"
 #include "Misc/EngineVersion.h"
+#include "Misc/FileHelper.h"
 #include "Misc/MessageDialog.h"
-#include "GenericPlatform/GenericPlatformInputDeviceMapper.h"
+#include "MotionControllerComponent.h"
+#include "SteamVRControllerKeys.h"
+#include "SteamVRSkeletonDefinition.h"
+#include "openvr.h"
 
 #if PLATFORM_WINDOWS
 #include "Windows/WindowsHWrapper.h"
diff --git a/Engine/Plugins/Runtime/Steam/SteamVR/Source/SteamVRInputDevice/Private/SteamVRInputDeviceModule.cpp b/Engine/Plugins/Runtime/Steam/SteamVR/Source/SteamVRInputDevice/Private/SteamVRInputDeviceModule.cpp
--- a/Engine/Plugins/Runtime/Steam/SteamVR/Source/SteamVRInputDevice/Private/SteamVRInputDeviceModule.cpp
+++ b/Engine/Plugins/Runtime/Steam/SteamVR/Source/SteamVRInputDevice/Private/SteamVRInputDeviceModule.cpp
@@ -29,11 +29,11 @@ POSSIBILITY OF SUCH DAMAGE.
 */
 
 #include "ISteamVRInputDeviceModule.h"
-#include "IInputDevice.h"
-#include "Runtime/Projects/Public/Interfaces/IPluginManager.h"
-#include "SteamVRInputDevice.h"
 #include "Engine/Engine.h"
+#include "IInputDevice.h"
 #include "IXRTrackingSystem.h"
+#include "Interfaces/IPluginManager.h"
+#include "SteamVRInputDevice.h"
 
 #if WITH_EDITOR
 #include "Editor/EditorEngine.h"
@@ -93,4 +93,4 @@ void FSteamVRInputDeviceModule::ShutdownModule()
 
 IMPLEMENT_MODULE(FSteamVRInputDeviceModule, SteamVRInputDevice)
 
-PRAGMA_ENABLE_DEPRECATION_WARNINGS
\ No newline at end of file
+PRAGMA_ENABLE_DEPRECATION_WARNINGS
diff --git a/Engine/Plugins/Runtime/ZoneGraph/Source/ZoneGraphEditor/Private/ZoneGraphEditorModule.cpp b/Engine/Plugins/Runtime/ZoneGraph/Source/ZoneGraphEditor/Private/ZoneGraphEditorModule.cpp
--- a/Engine/Plugins/Runtime/ZoneGraph/Source/ZoneGraphEditor/Private/ZoneGraphEditorModule.cpp
+++ b/Engine/Plugins/Runtime/ZoneGraph/Source/ZoneGraphEditor/Private/ZoneGraphEditorModule.cpp
@@ -1,30 +1,30 @@
 // Copyright Epic Games, Inc. All Rights Reserved.
 
 #include "ZoneGraphEditorModule.h"
-#include "Modules/ModuleManager.h"
-#include "Developer/AssetTools/Public/IAssetTools.h"
-#include "Developer/AssetTools/Public/AssetToolsModule.h"
-#include "ShowFlags.h"
-#include "ZoneShapeComponent.h"
-#include "ZoneShapeComponentDetails.h"
-#include "ZoneShapeComponentVisualizer.h"
-#include "ZoneLaneProfileRefDetails.h"
-#include "ZoneLaneProfileDetails.h"
-#include "ZoneLaneDescDetails.h"
-#include "ZoneGraphTagDetails.h"
-#include "ZoneGraphTagMaskDetails.h"
-#include "ZoneGraphTagInfoDetails.h"
-#include "ZoneGraphTagFilterDetails.h"
-#include "ZoneGraphTessellationSettingsDetails.h"
-#include "DetailLayoutBuilder.h"
+#include "AssetToolsModule.h"
 #include "DetailCategoryBuilder.h"
+#include "DetailLayoutBuilder.h"
 #include "Editor/UnrealEdEngine.h"
-#include "UnrealEdGlobals.h"
-#include "ToolMenus.h"
 #include "Framework/Commands/Commands.h"
+#include "IAssetTools.h"
+#include "Modules/ModuleManager.h"
+#include "ShowFlags.h"
 #include "Styling/AppStyle.h"
+#include "ToolMenus.h"
+#include "UnrealEdGlobals.h"
 #include "ZoneGraphDelegates.h"
 #include "ZoneGraphEditorStyle.h"
+#include "ZoneGraphTagDetails.h"
+#include "ZoneGraphTagFilterDetails.h"
+#include "ZoneGraphTagInfoDetails.h"
+#include "ZoneGraphTagMaskDetails.h"
+#include "ZoneGraphTessellationSettingsDetails.h"
+#include "ZoneLaneDescDetails.h"
+#include "ZoneLaneProfileDetails.h"
+#include "ZoneLaneProfileRefDetails.h"
+#include "ZoneShapeComponent.h"
+#include "ZoneShapeComponentDetails.h"
+#include "ZoneShapeComponentVisualizer.h"
 
 
 #define LOCTEXT_NAMESPACE "ZoneGraphEditor"
diff --git a/Engine/Plugins/VirtualProduction/LevelSnapshots/Source/LevelSnapshotsEditor/Private/Views/Results/SLevelSnapshotsEditorResults.cpp b/Engine/Plugins/VirtualProduction/LevelSnapshots/Source/LevelSnapshotsEditor/Private/Views/Results/SLevelSnapshotsEditorResults.cpp
--- a/Engine/Plugins/VirtualProduction/LevelSnapshots/Source/LevelSnapshotsEditor/Private/Views/Results/SLevelSnapshotsEditorResults.cpp
+++ b/Engine/Plugins/VirtualProduction/LevelSnapshots/Source/LevelSnapshotsEditor/Private/Views/Results/SLevelSnapshotsEditorResults.cpp
@@ -10,17 +10,17 @@
 #include "SnapshotRestorability.h"
 
 #include "Algo/Find.h"
-#include "Developer/ToolWidgets/Public/SPrimaryButton.h"
 #include "Framework/MultiBox/MultiBoxBuilder.h"
 #include "Framework/Notifications/NotificationManager.h"
 #include "IPropertyRowGenerator.h"
+#include "SPrimaryButton.h"
 #include "UObject/UObjectHash.h"
 #include "Widgets/Images/SImage.h"
 #include "Widgets/Input/SButton.h"
 #include "Widgets/Input/SComboButton.h"
 #include "Widgets/Input/SSearchBox.h"
-#include "Widgets/SBoxPanel.h"
 #include "Widgets/Notifications/SNotificationList.h"
+#include "Widgets/SBoxPanel.h"
 #include "Widgets/Text/STextBlock.h"
 
 #define LOCTEXT_NAMESPACE "LevelSnapshotsEditor"
diff --git a/Engine/Plugins/VirtualProduction/LevelSnapshots/Source/LevelSnapshotsEditor/Private/Widgets/SLevelSnapshotsEditorBrowser.cpp b/Engine/Plugins/VirtualProduction/LevelSnapshots/Source/LevelSnapshotsEditor/Private/Widgets/SLevelSnapshotsEditorBrowser.cpp
--- a/Engine/Plugins/VirtualProduction/LevelSnapshots/Source/LevelSnapshotsEditor/Private/Widgets/SLevelSnapshotsEditorBrowser.cpp
+++ b/Engine/Plugins/VirtualProduction/LevelSnapshots/Source/LevelSnapshotsEditor/Private/Widgets/SLevelSnapshotsEditorBrowser.cpp
@@ -13,7 +13,7 @@
 #include "AssetThumbnail.h"
 #include "ContentBrowserDelegates.h"
 #include "ContentBrowserModule.h"
-#include "Editor/EditorStyle/Public/EditorStyleSet.h"
+#include "EditorStyleSet.h"
 #include "Framework/MultiBox/MultiBoxBuilder.h"
 #include "IContentBrowserSingleton.h"
 #include "Misc/ScopedSlowTask.h"
@@ -21,9 +21,9 @@
 #include "Toolkits/GlobalEditorCommonCommands.h"
 #include "UObject/Object.h"
 #include "UObject/SoftObjectPath.h"
+#include "Widgets/Input/SSearchBox.h"
 #include "Widgets/Input/SSlider.h"
 #include "Widgets/Layout/SBox.h"
-#include "Widgets/Input/SSearchBox.h"
 #include "Widgets/Layout/SWidgetSwitcher.h"
 #include "Widgets/Views/SListView.h"
 
