Fix for ZoneGraph not automatically rebuilding in Editor when the Build Zone Graph While Editing flag is enabled.

diff --git a/Engine/Plugins/Runtime/ZoneGraph/Source/ZoneGraph/Private/ZoneGraphSubsystem.cpp b/Engine/Plugins/Runtime/ZoneGraph/Source/ZoneGraph/Private/ZoneGraphSubsystem.cpp
--- a/Engine/Plugins/Runtime/ZoneGraph/Source/ZoneGraph/Private/ZoneGraphSubsystem.cpp
+++ b/Engine/Plugins/Runtime/ZoneGraph/Source/ZoneGraph/Private/ZoneGraphSubsystem.cpp
@@ -22,6 +22,8 @@ UZoneGraphSubsystem::UZoneGraphSubsystem()
 
 void UZoneGraphSubsystem::Initialize(FSubsystemCollectionBase& Collection)
 {
+	Super::Initialize(Collection);
+
 #if WITH_EDITOR
 	OnActorMovedHandle = GEngine->OnActorMoved().AddUObject(this, &UZoneGraphSubsystem::OnActorMoved);
 	OnRequestRebuildHandle = UE::ZoneGraphDelegates::OnZoneGraphRequestRebuild.AddUObject(this, &UZoneGraphSubsystem::OnRequestRebuild);
diff --git a/Engine/Plugins/Runtime/ZoneGraph/Source/ZoneGraph/Private/ZoneShapeComponent.cpp b/Engine/Plugins/Runtime/ZoneGraph/Source/ZoneGraph/Private/ZoneShapeComponent.cpp
--- a/Engine/Plugins/Runtime/ZoneGraph/Source/ZoneGraph/Private/ZoneShapeComponent.cpp
+++ b/Engine/Plugins/Runtime/ZoneGraph/Source/ZoneGraph/Private/ZoneShapeComponent.cpp
@@ -574,7 +574,12 @@ void UZoneShapeComponent::ClearPerPointLaneProfiles()
 uint32 UZoneShapeComponent::GetShapeHash() const
 {
 	FZoneGraphObjectCRC32 Archive;
-	return Archive.Crc32(const_cast<UObject*>((const UObject*)this), 0);
+	const uint32 ShapeHash = Archive.Crc32(const_cast<UObject*>((const UObject*)this), 0);
+
+	const FTransform Transform = GetComponentTransform();
+	const uint32 TransformHash = GetTypeHash(Transform);
+
+	return HashCombine(ShapeHash, TransformHash);
 }
 #endif
 
diff --git a/Engine/Plugins/Runtime/ZoneGraph/Source/ZoneGraph/Public/ZoneGraphSubsystem.h b/Engine/Plugins/Runtime/ZoneGraph/Source/ZoneGraph/Public/ZoneGraphSubsystem.h
--- a/Engine/Plugins/Runtime/ZoneGraph/Source/ZoneGraph/Public/ZoneGraphSubsystem.h
+++ b/Engine/Plugins/Runtime/ZoneGraph/Source/ZoneGraph/Public/ZoneGraphSubsystem.h
@@ -124,6 +124,10 @@ public:
 	// Returns all tag infos
 	TConstArrayView<FZoneGraphTagInfo> GetTagInfos() const;
 
+#if WITH_EDITOR
+	virtual bool IsTickableInEditor() const override { return true; }
+#endif /* WITH_EDITOR */
+
 protected:
 	virtual void Initialize(FSubsystemCollectionBase& Collection) override;
 	virtual void PostInitialize() override;
