Adding dependency back in the correct format.

diff --git a/Engine/Plugins/Runtime/ZoneGraph/Source/ZoneGraph/ZoneGraph.Build.cs b/Engine/Plugins/Runtime/ZoneGraph/Source/ZoneGraph/ZoneGraph.Build.cs
--- a/Engine/Plugins/Runtime/ZoneGraph/Source/ZoneGraph/ZoneGraph.Build.cs
+++ b/Engine/Plugins/Runtime/ZoneGraph/Source/ZoneGraph/ZoneGraph.Build.cs
@@ -8,6 +8,8 @@ namespace UnrealBuildTool.Rules
 		{
 			PCHUsage = ModuleRules.PCHUsageMode.UseExplicitOrSharedPCHs;
 
+			PublicIncludePathModuleNames.Add("AIModule");
+
 			PublicDependencyModuleNames.AddRange(
 				new string[] {
 					"Core",
