Added FZoneGraphTagFilter to FZoneGraphPathFilter so that FZoneGraphAStarNode can perform pathing requests using a given tag filter (and updated UZoneGraphTestingComponent to use it with its QueryFilter).

diff --git a/Engine/Plugins/Runtime/ZoneGraph/Source/ZoneGraph/Public/ZoneGraphAStar.h b/Engine/Plugins/Runtime/ZoneGraph/Source/ZoneGraph/Public/ZoneGraphAStar.h
--- a/Engine/Plugins/Runtime/ZoneGraph/Source/ZoneGraph/Public/ZoneGraphAStar.h
+++ b/Engine/Plugins/Runtime/ZoneGraph/Source/ZoneGraph/Public/ZoneGraphAStar.h
@@ -59,8 +59,9 @@ struct FZoneGraphPathFilter
 
 	typedef int32/*lane index type*/ FNodeRef;
 
-	FZoneGraphPathFilter(const FZoneGraphStorage& InGraph, const FZoneGraphLaneLocation& InStartLocation, const FZoneGraphLaneLocation& InEndLocation)
+	FZoneGraphPathFilter(const FZoneGraphStorage& InGraph, const FZoneGraphLaneLocation& InStartLocation, const FZoneGraphLaneLocation& InEndLocation, const FZoneGraphTagFilter InZoneTagFilter = FZoneGraphTagFilter())
 		: ZoneStorage(InGraph) 
+		, ZoneTagFilter(InZoneTagFilter)
 		, StartLocation(InStartLocation)
 		, EndLocation(InEndLocation)
 	{}
@@ -164,12 +165,18 @@ struct FZoneGraphPathFilter
 		}
 	}
 
-	FORCEINLINE bool IsTraversalAllowed(const FNodeRef StartNodeRef, const FNodeRef& Neighbour) const { return true; }
+	FORCEINLINE bool IsTraversalAllowed(const FNodeRef StartNodeRef, const FNodeRef& Neighbour) const
+	{
+		const FZoneGraphTagMask& LaneTagMask = ZoneStorage.Lanes[Neighbour].Tags;
+		return ZoneTagFilter.Pass(LaneTagMask);
+	}
+
 	FORCEINLINE bool WantsPartialSolution() const { return false; }
 	FORCEINLINE bool ShouldIncludeStartNodeInPath() const { return true; }
 
 protected:
 	const FZoneGraphStorage& ZoneStorage;
+	const FZoneGraphTagFilter ZoneTagFilter;
 	const FZoneGraphLaneLocation StartLocation;
 	const FZoneGraphLaneLocation EndLocation;
 };
diff --git a/Engine/Plugins/Runtime/ZoneGraph/Source/ZoneGraphDebug/Private/ZoneGraphTestingActor.cpp b/Engine/Plugins/Runtime/ZoneGraph/Source/ZoneGraphDebug/Private/ZoneGraphTestingActor.cpp
--- a/Engine/Plugins/Runtime/ZoneGraph/Source/ZoneGraphDebug/Private/ZoneGraphTestingActor.cpp
+++ b/Engine/Plugins/Runtime/ZoneGraph/Source/ZoneGraphDebug/Private/ZoneGraphTestingActor.cpp
@@ -170,7 +170,7 @@ void UZoneGraphTestingComponent::UpdateTests()
 				// @todo: pass FZoneGraphLaneLocation directly to the constructor
 				FZoneGraphAStarNode StartNode(NearestLaneLocation.LaneHandle.Index, NearestLaneLocation.Position);
 				FZoneGraphAStarNode EndNode(OtherTestingComp->NearestLaneLocation.LaneHandle.Index, OtherTestingComp->NearestLaneLocation.Position);
-				FZoneGraphPathFilter PathFilter(ZoneGraphStorage, NearestLaneLocation, OtherTestingComp->NearestLaneLocation);
+				FZoneGraphPathFilter PathFilter(ZoneGraphStorage, NearestLaneLocation, OtherTestingComp->NearestLaneLocation, QueryFilter);
 				
 				// @todo: see if we can return directly a path of lane handles
 				TArray<FZoneGraphAStarWrapper::FNodeRef> ResultPath;
