Fix ZoneShape visualization when re-selecting shape or undoing actions.

diff --git a/Engine/Plugins/Runtime/ZoneGraph/Source/ZoneGraphEditor/Private/ZoneShapeComponentVisualizer.cpp b/Engine/Plugins/Runtime/ZoneGraph/Source/ZoneGraphEditor/Private/ZoneShapeComponentVisualizer.cpp
--- a/Engine/Plugins/Runtime/ZoneGraph/Source/ZoneGraphEditor/Private/ZoneShapeComponentVisualizer.cpp
+++ b/Engine/Plugins/Runtime/ZoneGraph/Source/ZoneGraphEditor/Private/ZoneShapeComponentVisualizer.cpp
@@ -1125,7 +1125,7 @@ bool FZoneShapeComponentVisualizer::HandleSnapTo(const bool bInAlign, const bool
 void FZoneShapeComponentVisualizer::EndEditing()
 {
 	// Ignore if there is an undo/redo operation in progress
-	if (!GIsTransacting)
+	if (GIsTransacting)
 	{
 		return;
 	}
