[ZoneGraph] added specifiers to property PerPointLaneProfile so it is properly handled by the component instance data when a BP actor is reconstructed.

diff --git a/Engine/Plugins/Runtime/ZoneGraph/Source/ZoneGraph/Public/ZoneShapeComponent.h b/Engine/Plugins/Runtime/ZoneGraph/Source/ZoneGraph/Public/ZoneShapeComponent.h
--- a/Engine/Plugins/Runtime/ZoneGraph/Source/ZoneGraph/Public/ZoneShapeComponent.h
+++ b/Engine/Plugins/Runtime/ZoneGraph/Source/ZoneGraph/Public/ZoneShapeComponent.h
@@ -184,7 +184,7 @@ private:
 	bool bReverseLaneProfile = false;
 
 	/** Array of lane templates indexed by the points when the shape is polygon. */
-	UPROPERTY(meta = (IncludeInHash))
+	UPROPERTY(Category = Zone, VisibleAnywhere, meta = (IncludeInHash, EditCondition = "false", EditConditionHides))
 	TArray<FZoneLaneProfileRef> PerPointLaneProfiles;
 
 	/** Shape points */
