ZoneGraph: Break a spline zone shape at selected control points or cursor position.

diff --git a/Engine/Plugins/Runtime/ZoneGraph/Source/ZoneGraphEditor/Private/ZoneShapeComponentVisualizer.cpp b/Engine/Plugins/Runtime/ZoneGraph/Source/ZoneGraphEditor/Private/ZoneShapeComponentVisualizer.cpp
--- a/Engine/Plugins/Runtime/ZoneGraph/Source/ZoneGraphEditor/Private/ZoneShapeComponentVisualizer.cpp
+++ b/Engine/Plugins/Runtime/ZoneGraph/Source/ZoneGraphEditor/Private/ZoneShapeComponentVisualizer.cpp
@@ -14,11 +14,13 @@
 #include "Editor.h"
 #include "EditorViewportClient.h"
 #include "EditorViewportCommands.h"
+#include "LevelEditor.h"
 #include "LevelEditorActions.h"
 #include "ScopedTransaction.h"
 #include "ActorEditorUtils.h"
 #include "ZoneGraphSubsystem.h"
 #include "ZoneGraphSettings.h"
+#include "ZoneShapeActor.h"
 #include "ZoneShapeComponent.h"
 #include "ZoneShapeUtilities.h"
 #include "ZoneGraphRenderingUtilities.h"
@@ -59,6 +61,10 @@ public:
 		UI_COMMAND(SetPointToAutoBezier, "Auto Bezier", "Set point to Auto Bezier type", EUserInterfaceActionType::RadioButton, FInputChord());
 		UI_COMMAND(SetPointToLaneSegment, "Lane Segment", "Set point to Lane Segment type", EUserInterfaceActionType::RadioButton, FInputChord());
 		UI_COMMAND(FocusViewportToSelection, "Focus Selected", "Moves the camera in front of the selection", EUserInterfaceActionType::Button, FInputChord(EKeys::F));
+		UI_COMMAND(BreakAtPointNewActors, "Break Into Shape Actors At Point(s)", "Break the shape into multiple shape actors at the currently selected points.", EUserInterfaceActionType::Button, FInputChord());
+		UI_COMMAND(BreakAtPointNewComponents, "Break Into Shape Components At Point(s)", "Break the shape into multiple shape components at the currently selected points.", EUserInterfaceActionType::Button, FInputChord());
+		UI_COMMAND(BreakAtSegmentNewActors, "Break Into Shape Actors Here", "Break the shape into multiple shape actors at the cursor location.", EUserInterfaceActionType::Button, FInputChord());
+		UI_COMMAND(BreakAtSegmentNewComponents, "Break Into Shape Components Here", "Break the shape into multiple shape components at the cursor location.", EUserInterfaceActionType::Button, FInputChord());
 	}
 
 public:
@@ -71,6 +77,10 @@ public:
 	TSharedPtr<FUICommandInfo> SetPointToAutoBezier;
 	TSharedPtr<FUICommandInfo> SetPointToLaneSegment;
 	TSharedPtr<FUICommandInfo> FocusViewportToSelection;
+	TSharedPtr<FUICommandInfo> BreakAtPointNewActors;
+	TSharedPtr<FUICommandInfo> BreakAtPointNewComponents;
+	TSharedPtr<FUICommandInfo> BreakAtSegmentNewActors;
+	TSharedPtr<FUICommandInfo> BreakAtSegmentNewComponents;
 };
 
 FZoneShapeComponentVisualizer::FZoneShapeComponentVisualizer()
@@ -142,6 +152,26 @@ void FZoneShapeComponentVisualizer::OnRegister()
 		FExecuteAction::CreateStatic(&FLevelEditorActionCallbacks::ExecuteExecCommand, FString(TEXT("CAMERA ALIGN ACTIVEVIEWPORTONLY")))
 	);
 
+	ShapeComponentVisualizerActions->MapAction(
+		Commands.BreakAtPointNewActors,
+		FExecuteAction::CreateSP(this, &FZoneShapeComponentVisualizer::OnBreakAtPointNewActors),
+		FCanExecuteAction::CreateSP(this, &FZoneShapeComponentVisualizer::CanBreakAtPoint));
+
+	ShapeComponentVisualizerActions->MapAction(
+		Commands.BreakAtPointNewComponents,
+		FExecuteAction::CreateSP(this, &FZoneShapeComponentVisualizer::OnBreakAtPointNewComponents),
+		FCanExecuteAction::CreateSP(this, &FZoneShapeComponentVisualizer::CanBreakAtPoint));
+
+	ShapeComponentVisualizerActions->MapAction(
+		Commands.BreakAtSegmentNewActors,
+		FExecuteAction::CreateSP(this, &FZoneShapeComponentVisualizer::OnBreakAtSegmentNewActors),
+		FCanExecuteAction::CreateSP(this, &FZoneShapeComponentVisualizer::CanBreakAtSegment));
+
+	ShapeComponentVisualizerActions->MapAction(
+		Commands.BreakAtSegmentNewComponents,
+		FExecuteAction::CreateSP(this, &FZoneShapeComponentVisualizer::OnBreakAtSegmentNewComponents),
+		FCanExecuteAction::CreateSP(this, &FZoneShapeComponentVisualizer::CanBreakAtSegment));
+
 	bool bAlign = false;
 	bool bUseLineTrace = false;
 	bool bUseBounds = false;
@@ -1557,6 +1587,181 @@ bool FZoneShapeComponentVisualizer::CanSelectAllPoints() const
 	return (ShapeComp != nullptr);
 }
 
+void FZoneShapeComponentVisualizer::OnBreakAtPointNewActors() const
+{
+	const FScopedTransaction Transaction(LOCTEXT("BreakAtPointNewActors", "Break Shape Into New Actors At Points"));
+	BreakAtPoint(true);
+}
+
+void FZoneShapeComponentVisualizer::OnBreakAtPointNewComponents() const
+{
+	const FScopedTransaction Transaction(LOCTEXT("BreakAtPointNewComponents", "Break Shape Into New Components At Points"));
+	BreakAtPoint(false);
+}
+
+void FZoneShapeComponentVisualizer::BreakAtPoint(bool bCreateNewActor) const
+{
+	UZoneShapeComponent* ShapeComp = GetEditedShapeComponent();
+	check(ShapeComp != nullptr);
+	check(SelectionState);
+	const TSet<int32>& SelectedPoints = SelectionState->GetSelectedPoints();
+	const int32 LastPointIndexSelected = SelectionState->GetLastPointIndexSelected();
+	check(LastPointIndexSelected != INDEX_NONE);
+	check(LastPointIndexSelected >= 0);
+	check(LastPointIndexSelected < ShapeComp->GetNumPoints());
+	check(SelectedPoints.Num() > 0);
+	check(SelectedPoints.Contains(LastPointIndexSelected));
+
+	ShapeComp->Modify();
+	if (AActor* Owner = ShapeComp->GetOwner())
+	{
+		Owner->Modify();
+	}
+
+	// Get a sorted list of all the selected indices, highest to lowest
+	TArray<int32> SelectedPointsSorted;
+	for (int32 SelectedIndex : SelectedPoints)
+	{
+		SelectedPointsSorted.Add(SelectedIndex);
+	}
+	SelectedPointsSorted.Sort([](int32 A, int32 B)
+		{ return A < B; });
+
+	// Create a new shape and then delete selected key from list, highest index first
+	FActorSpawnParameters SpawnParams;
+	TArray<FZoneShapePoint>& ShapePoints = ShapeComp->GetMutablePoints();
+	int32 EndIndex = ShapePoints.Num() - 1;
+	for (int32 i = SelectedPointsSorted.Num() - 1; i >= 0; i--)
+	{
+		if (ShapePoints.Num() <= 2)
+		{
+			// Keep at least 2 points
+			break;
+		}
+
+		const int32 SelectedIndex = SelectedPointsSorted[i];
+		if (SelectedIndex == (ShapePoints.Num() - 1) || SelectedIndex == 0)
+		{
+			continue;
+		}
+
+		// Create a new shape
+		UZoneShapeComponent* NewShapeComponent = nullptr;
+		AActor* ShapeOwner = ShapeComp->GetOwner();
+		if (bCreateNewActor)
+		{
+			AZoneShape* NewShapeActor = ShapeComp->GetWorld()->SpawnActor<AZoneShape>(AZoneShape::StaticClass(), ShapeComp->GetComponentTransform(), SpawnParams);
+			if (!NewShapeActor)
+			{
+				continue;
+			}
+			NewShapeComponent = NewShapeActor->GetComponentByClass<UZoneShapeComponent>();
+			NewShapeActor->Modify();
+		}
+		else
+		{
+			NewShapeComponent = NewObject<UZoneShapeComponent>(ShapeComp->GetOuter(), NAME_None, RF_Transactional);
+			if (!NewShapeComponent)
+			{
+				continue;
+			}
+			NewShapeComponent->SetWorldTransform(ShapeComp->GetComponentTransform());
+			ShapeOwner->AddInstanceComponent(NewShapeComponent);
+			NewShapeComponent->RegisterComponent();
+			NewShapeComponent->AttachToComponent(ShapeComp, FAttachmentTransformRules::KeepWorldTransform); // Should we attach to the root component?
+			NewShapeComponent->Modify();
+		}
+
+		// Copy points
+		TArray<FZoneShapePoint>& NewShapePoints = NewShapeComponent->GetMutablePoints();
+		NewShapePoints.SetNum(EndIndex - SelectedIndex + 1);
+		int32 SrcIndex = SelectedIndex;
+		for (int32 Index = 0; Index < NewShapePoints.Num(); Index++, SrcIndex++)
+		{
+			NewShapePoints[Index] = ShapePoints[SrcIndex];
+		}
+		NewShapeComponent->UpdateShape();
+
+		if (i == 0 || (i == 1 && SelectedPointsSorted[0] == 0))
+		{
+			// Keep the last segment on the original shape component
+			ShapePoints.RemoveAt(EndIndex);
+			break;
+		}
+
+		// Delete all points after the selected one
+		for (int32 Index = EndIndex; Index > SelectedIndex; Index--)
+		{
+			if (Index <= 1)
+			{
+				// The zone shape needs at least two points
+				break;
+			}
+			ShapePoints.RemoveAt(Index);
+		}
+		EndIndex = SelectedIndex;
+	}
+
+	// Clear selection
+	ChangeSelectionState(INDEX_NONE, false);
+	SelectionState->SetSelectedSegmentIndex(INDEX_NONE);
+	SelectionState->SetSelectedControlPoint(INDEX_NONE);
+	SelectionState->SetSelectedControlPointType(FZoneShapeControlPointType::None);
+
+	ShapeComp->UpdateShape();
+	NotifyPropertyModified(ShapeComp, ShapePointsProperty);
+
+	GEditor->RedrawLevelEditingViewports(true);
+	FLevelEditorModule& LevelEditor = FModuleManager::LoadModuleChecked<FLevelEditorModule>("LevelEditor");
+	LevelEditor.BroadcastComponentsEdited();
+	LevelEditor.BroadcastRedrawViewports(false);
+}
+
+bool FZoneShapeComponentVisualizer::CanBreakAtPoint() const
+{
+	check(SelectionState);
+	const TSet<int32>& SelectedPoints = SelectionState->GetSelectedPoints();
+	const int32 LastPointIndexSelected = SelectionState->GetLastPointIndexSelected();
+	UZoneShapeComponent* ShapeComp = GetEditedShapeComponent();
+	return (ShapeComp != nullptr && SelectedPoints.Num() > 0 && LastPointIndexSelected != INDEX_NONE);
+}
+
+void FZoneShapeComponentVisualizer::OnBreakAtSegmentNewActors() const
+{
+	const FScopedTransaction Transaction(LOCTEXT("BreakAtSegmentNewActors", "Break Shape Into New Actors At The Cursor Location"));
+	BreakAtSegment(true);
+}
+
+void FZoneShapeComponentVisualizer::OnBreakAtSegmentNewComponents() const
+{
+	const FScopedTransaction Transaction(LOCTEXT("BreakAtSegmentNewComponents", "Break Shape Into New Components At The Cursor Location"));
+	BreakAtSegment(false);
+}
+
+void FZoneShapeComponentVisualizer::BreakAtSegment(bool bCreateNewActor) const
+{
+	UZoneShapeComponent* ShapeComp = GetEditedShapeComponent();
+	check(ShapeComp != nullptr);
+	const int32 SelectedSegmentIndex = SelectionState->GetSelectedSegmentIndex();
+	check(SelectionState);
+	check(SelectedSegmentIndex != INDEX_NONE);
+	check(SelectedSegmentIndex >= 0);
+	check(SelectedSegmentIndex < ShapeComp->GetNumSegments());
+	SelectionState->Modify();
+	int32 SegmentIndex = SelectionState->GetSelectedSegmentIndex();
+	SplitSegment(SegmentIndex, SelectionState->GetSelectedSegmentT());
+	const int NewPointIndex = SegmentIndex + 1;
+	ChangeSelectionState(NewPointIndex, false);
+	BreakAtPoint(bCreateNewActor);
+	SelectionState->SetSelectedSegmentPoint(FVector::ZeroVector);
+	SelectionState->SetSelectedSegmentIndex(INDEX_NONE);
+}
+
+bool FZoneShapeComponentVisualizer::CanBreakAtSegment() const
+{
+	return CanAddPointToSegment();
+}
+
 TSharedPtr<SWidget> FZoneShapeComponentVisualizer::GenerateContextMenu() const
 {
 	check(SelectionState);
@@ -1568,6 +1773,11 @@ TSharedPtr<SWidget> FZoneShapeComponentVisualizer::GenerateContextMenu() const
 		if (SelectionState->GetSelectedSegmentIndex() != INDEX_NONE)
 		{
 			MenuBuilder.AddMenuEntry(FZoneShapeComponentVisualizerCommands::Get().AddPoint);
+
+			MenuBuilder.AddSubMenu(
+				LOCTEXT("BreakAtPoint", "Break At Point"),
+				LOCTEXT("BreakAtPointTooltip", "Break the shape into pieces at the currently selected points."),
+				FNewMenuDelegate::CreateSP(this, &FZoneShapeComponentVisualizer::GenerateBreakAtSegmentSubMenu));
 		}
 		else if (SelectionState->GetLastPointIndexSelected() != INDEX_NONE)
 		{
@@ -1579,6 +1789,11 @@ TSharedPtr<SWidget> FZoneShapeComponentVisualizer::GenerateContextMenu() const
 				LOCTEXT("ShapePointType", "Point Type"),
 				LOCTEXT("ShapePointTypeTooltip", "Define the type of the point."),
 				FNewMenuDelegate::CreateSP(this, &FZoneShapeComponentVisualizer::GenerateShapePointTypeSubMenu));
+
+			MenuBuilder.AddSubMenu(
+				LOCTEXT("BreakAtPoint", "Break At Point"),
+				LOCTEXT("BreakAtPointTooltip", "Break the shape into pieces at the currently selected points."),
+				FNewMenuDelegate::CreateSP(this, &FZoneShapeComponentVisualizer::GenerateBreakAtPointSubMenu));
 		}
 	}
 	MenuBuilder.EndSection();
@@ -1606,4 +1821,24 @@ void FZoneShapeComponentVisualizer::GenerateShapePointTypeSubMenu(FMenuBuilder&
 	}
 }
 
+void FZoneShapeComponentVisualizer::GenerateBreakAtPointSubMenu(FMenuBuilder& MenuBuilder) const
+{
+	UZoneShapeComponent* ShapeComp = GetEditedShapeComponent();
+	if (ShapeComp && ShapeComp->GetShapeType() == FZoneShapeType::Spline)
+	{
+		MenuBuilder.AddMenuEntry(FZoneShapeComponentVisualizerCommands::Get().BreakAtPointNewActors);
+		MenuBuilder.AddMenuEntry(FZoneShapeComponentVisualizerCommands::Get().BreakAtPointNewComponents);
+	}
+}
+
+void FZoneShapeComponentVisualizer::GenerateBreakAtSegmentSubMenu(FMenuBuilder& MenuBuilder) const
+{
+	UZoneShapeComponent* ShapeComp = GetEditedShapeComponent();
+	if (ShapeComp && ShapeComp->GetShapeType() == FZoneShapeType::Spline)
+	{
+		MenuBuilder.AddMenuEntry(FZoneShapeComponentVisualizerCommands::Get().BreakAtSegmentNewActors);
+		MenuBuilder.AddMenuEntry(FZoneShapeComponentVisualizerCommands::Get().BreakAtSegmentNewComponents);
+	}
+}
+
 #undef LOCTEXT_NAMESPACE
diff --git a/Engine/Plugins/Runtime/ZoneGraph/Source/ZoneGraphEditor/Private/ZoneShapeComponentVisualizer.h b/Engine/Plugins/Runtime/ZoneGraph/Source/ZoneGraphEditor/Private/ZoneShapeComponentVisualizer.h
--- a/Engine/Plugins/Runtime/ZoneGraph/Source/ZoneGraphEditor/Private/ZoneShapeComponentVisualizer.h
+++ b/Engine/Plugins/Runtime/ZoneGraph/Source/ZoneGraphEditor/Private/ZoneShapeComponentVisualizer.h
@@ -238,8 +238,19 @@ protected:
 	void OnSelectAllPoints() const;
 	bool CanSelectAllPoints() const;
 
+	void OnBreakAtPointNewActors() const;
+	void OnBreakAtPointNewComponents() const;
+	bool CanBreakAtPoint() const;
+
+	void OnBreakAtSegmentNewActors() const;
+	void OnBreakAtSegmentNewComponents() const;
+	bool CanBreakAtSegment() const;
+
 	void GenerateShapePointTypeSubMenu(FMenuBuilder& MenuBuilder) const;
 
+	void GenerateBreakAtPointSubMenu(FMenuBuilder& MenuBuilder) const;
+	void GenerateBreakAtSegmentSubMenu(FMenuBuilder& MenuBuilder) const;
+
 	// FGCObject interface
 	virtual void AddReferencedObjects(FReferenceCollector& Collector) override;
 	virtual FString GetReferencerName() const override
@@ -276,4 +287,8 @@ protected:
 	FQuat CachedRotation = FQuat::Identity;
 
 	bool bIsSelectingComponent = false;
+
+private:
+	void BreakAtPoint(bool bCreateNewActor) const;
+	void BreakAtSegment(bool bCreateNewActor) const;
 };
