ZoneGraph: Editor snapping improvements
- Click to add a control point to the zone shape and snap it to the ground while hold the V key.
- Add a submenu to snap/align selected points to the ground.

diff --git a/Engine/Plugins/Runtime/ZoneGraph/Source/ZoneGraphEditor/Private/ZoneShapeComponentVisualizer.cpp b/Engine/Plugins/Runtime/ZoneGraph/Source/ZoneGraphEditor/Private/ZoneShapeComponentVisualizer.cpp
--- a/Engine/Plugins/Runtime/ZoneGraph/Source/ZoneGraphEditor/Private/ZoneShapeComponentVisualizer.cpp
+++ b/Engine/Plugins/Runtime/ZoneGraph/Source/ZoneGraphEditor/Private/ZoneShapeComponentVisualizer.cpp
@@ -35,6 +35,7 @@ IMPLEMENT_HIT_PROXY(HZoneShapeSegmentProxy, HZoneShapeVisProxy);
 IMPLEMENT_HIT_PROXY(HZoneShapeControlPointProxy, HZoneShapeVisProxy);
 
 #define LOCTEXT_NAMESPACE "ZoneShapeComponentVisualizer"
+DEFINE_LOG_CATEGORY_STATIC(LogZoneShapeComponentVisualizer, Log, All)
 
 /** Define commands for the shape component visualizer */
 class FZoneShapeComponentVisualizerCommands : public TCommands<FZoneShapeComponentVisualizerCommands>
@@ -962,6 +963,43 @@ bool FZoneShapeComponentVisualizer::HandleInputKey(FEditorViewportClient* Viewpo
 
 	if (Event == IE_Pressed)
 	{
+		// Add a new point to the shape when you hold the V key and press left mouse button
+		if (ShapeComp && Key == EKeys::LeftMouseButton && Viewport->KeyState(EKeys::V))
+		{
+			// Get clicked position
+			UWorld* World = ViewportClient->GetWorld();
+			FSceneViewFamilyContext ViewFamily(FSceneViewFamilyContext::ConstructionValues(ViewportClient->Viewport, ViewportClient->GetScene(), ViewportClient->EngineShowFlags)
+				.SetRealtimeUpdate(ViewportClient->IsRealtime()));
+			FSceneView* View = ViewportClient->CalcSceneView(&ViewFamily);
+			int32 MouseX = ViewportClient->Viewport->GetMouseX();
+			int32 MouseY = ViewportClient->Viewport->GetMouseY();
+			FViewportCursorLocation MouseViewportRay(View, ViewportClient, MouseX, MouseY);
+			FVector MouseViewportRayDirection = MouseViewportRay.GetDirection();
+
+			FVector Start = MouseViewportRay.GetOrigin();
+			FVector End = Start + WORLD_MAX * MouseViewportRayDirection;
+			if (ViewportClient->IsOrtho())
+			{
+				Start -= WORLD_MAX * MouseViewportRayDirection;
+			}
+			FHitResult Hit;
+			FCollisionQueryParams QueryParams;
+			QueryParams.bTraceComplex = true;
+			if (World->LineTraceSingleByChannel(Hit, Start, End, ECollisionChannel::ECC_WorldStatic, QueryParams))
+			{
+				// Add a new point at the position
+				TArray<FZoneShapePoint>& Points = ShapeComp->GetMutablePoints();
+				FZoneShapePoint PointToAdd(ShapeComp->GetComponentTransform().InverseTransformPosition(Hit.Location));
+				Points.Add(PointToAdd);
+				ShapeComp->UpdateShape();
+			}
+			else
+			{
+				UE_LOG(LogZoneShapeComponentVisualizer, Warning, TEXT("No hit found on click."));
+			}
+			return true;
+		}
+
 		bHandled = ShapeComponentVisualizerActions->ProcessCommandBindings(Key, FSlateApplication::Get().GetModifierKeys(), false);
 	}
 
@@ -1790,6 +1828,11 @@ TSharedPtr<SWidget> FZoneShapeComponentVisualizer::GenerateContextMenu() const
 				LOCTEXT("ShapePointTypeTooltip", "Define the type of the point."),
 				FNewMenuDelegate::CreateSP(this, &FZoneShapeComponentVisualizer::GenerateShapePointTypeSubMenu));
 
+			MenuBuilder.AddSubMenu(
+				LOCTEXT("SplineSnapAlign", "Snap/Align"),
+				LOCTEXT("SplineSnapAlignTooltip", "Snap align options."),
+				FNewMenuDelegate::CreateSP(this, &FZoneShapeComponentVisualizer::GenerateSnapAlignSubMenu));
+
 			MenuBuilder.AddSubMenu(
 				LOCTEXT("BreakAtPoint", "Break At Point"),
 				LOCTEXT("BreakAtPointTooltip", "Break the shape into pieces at the currently selected points."),
@@ -1821,6 +1864,12 @@ void FZoneShapeComponentVisualizer::GenerateShapePointTypeSubMenu(FMenuBuilder&
 	}
 }
 
+void FZoneShapeComponentVisualizer::GenerateSnapAlignSubMenu(FMenuBuilder& MenuBuilder) const
+{
+	MenuBuilder.AddMenuEntry(FLevelEditorCommands::Get().SnapToFloor);
+	MenuBuilder.AddMenuEntry(FLevelEditorCommands::Get().AlignToFloor);
+}
+
 void FZoneShapeComponentVisualizer::GenerateBreakAtPointSubMenu(FMenuBuilder& MenuBuilder) const
 {
 	UZoneShapeComponent* ShapeComp = GetEditedShapeComponent();
diff --git a/Engine/Plugins/Runtime/ZoneGraph/Source/ZoneGraphEditor/Private/ZoneShapeComponentVisualizer.h b/Engine/Plugins/Runtime/ZoneGraph/Source/ZoneGraphEditor/Private/ZoneShapeComponentVisualizer.h
--- a/Engine/Plugins/Runtime/ZoneGraph/Source/ZoneGraphEditor/Private/ZoneShapeComponentVisualizer.h
+++ b/Engine/Plugins/Runtime/ZoneGraph/Source/ZoneGraphEditor/Private/ZoneShapeComponentVisualizer.h
@@ -248,6 +248,7 @@ protected:
 
 	void GenerateShapePointTypeSubMenu(FMenuBuilder& MenuBuilder) const;
 
+	void GenerateSnapAlignSubMenu(FMenuBuilder& MenuBuilder) const;
 	void GenerateBreakAtPointSubMenu(FMenuBuilder& MenuBuilder) const;
 	void GenerateBreakAtSegmentSubMenu(FMenuBuilder& MenuBuilder) const;
 
