[ZoneGraph] When dragging an end control point, show possible connection previews and auto connect to the closest one when released.

diff --git a/Engine/Plugins/Runtime/ZoneGraph/Source/ZoneGraph/Private/ZoneGraphBuilder.cpp b/Engine/Plugins/Runtime/ZoneGraph/Source/ZoneGraph/Private/ZoneGraphBuilder.cpp
--- a/Engine/Plugins/Runtime/ZoneGraph/Source/ZoneGraph/Private/ZoneGraphBuilder.cpp
+++ b/Engine/Plugins/Runtime/ZoneGraph/Source/ZoneGraph/Private/ZoneGraphBuilder.cpp
@@ -751,3 +751,8 @@ void FZoneGraphBuilder::ConnectLanes(TArray<FZoneShapeLaneInternalLink>& Interna
 		Lane.LinksEnd = ZoneStorage.LaneLinks.Num();
 	}
 }
+
+void FZoneGraphBuilder::QueryHashGrid(const FBox& Bounds, TArray<FZoneGraphBuilderHashGrid2D::ItemIDType>& OutResults)
+{
+	HashGrid.Query(Bounds, OutResults);
+}
diff --git a/Engine/Plugins/Runtime/ZoneGraph/Source/ZoneGraph/Private/ZoneShapeComponent.cpp b/Engine/Plugins/Runtime/ZoneGraph/Source/ZoneGraph/Private/ZoneShapeComponent.cpp
--- a/Engine/Plugins/Runtime/ZoneGraph/Source/ZoneGraph/Private/ZoneShapeComponent.cpp
+++ b/Engine/Plugins/Runtime/ZoneGraph/Source/ZoneGraph/Private/ZoneShapeComponent.cpp
@@ -865,5 +865,17 @@ FPrimitiveSceneProxy* UZoneShapeComponent::CreateSceneProxy()
 	return new FZoneShapeSceneProxy(*this);
 }
 
+const FZoneShapeConnector* UZoneShapeComponent::GetShapeConnectorByPointIndex(int32 PointIndex) const
+{
+	for (const FZoneShapeConnector& ShapeConnector : ShapeConnectors)
+	{
+		if (ShapeConnector.PointIndex == PointIndex)
+		{
+			return &ShapeConnector;
+		}
+	}
+	return nullptr;
+}
+
 #endif // !UE_BUILD_SHIPPING
 
diff --git a/Engine/Plugins/Runtime/ZoneGraph/Source/ZoneGraph/Public/ZoneGraphBuilder.h b/Engine/Plugins/Runtime/ZoneGraph/Source/ZoneGraph/Public/ZoneGraphBuilder.h
--- a/Engine/Plugins/Runtime/ZoneGraph/Source/ZoneGraph/Public/ZoneGraphBuilder.h
+++ b/Engine/Plugins/Runtime/ZoneGraph/Source/ZoneGraph/Public/ZoneGraphBuilder.h
@@ -86,6 +86,12 @@ public:
 	*/
 	static void BuildSingleShape(const UZoneShapeComponent& ShapeComp, const FMatrix& LocalToWorld, FZoneGraphStorage& OutZoneStorage);
 
+	/** Returns items that potentially touch the bounds in the HashGrid. Operates on grid level, can have false positives.
+	 * @param Bounds - Query bounding box.
+	 * @param OutResults - Result of the query, IDs of potentially overlapping items.
+	 */
+	void QueryHashGrid(const FBox& Bounds, TArray<FZoneGraphBuilderHashGrid2D::ItemIDType>& OutResults);
+
 protected:
 	void Build(AZoneGraphData& ZoneGraphData);
 	void RequestRebuild();
diff --git a/Engine/Plugins/Runtime/ZoneGraph/Source/ZoneGraph/Public/ZoneGraphTypes.h b/Engine/Plugins/Runtime/ZoneGraph/Source/ZoneGraph/Public/ZoneGraphTypes.h
--- a/Engine/Plugins/Runtime/ZoneGraph/Source/ZoneGraph/Public/ZoneGraphTypes.h
+++ b/Engine/Plugins/Runtime/ZoneGraph/Source/ZoneGraph/Public/ZoneGraphTypes.h
@@ -995,4 +995,8 @@ struct ZONEGRAPH_API FZoneGraphBuildSettings
 	/** Max relative angle (in degrees) between two shape points for them to be snapped together. */
 	UPROPERTY(Category = PointSnapping, EditAnywhere)
 	float ConnectionSnapAngle = 10.0f;
+
+	/** Max distance to auto connect a dragged end control point to the closest connector when released. */
+	UPROPERTY(Category = PointSnapping, EditAnywhere)
+	float DragEndpointAutoConnectRange = 250.0f;
 };
diff --git a/Engine/Plugins/Runtime/ZoneGraph/Source/ZoneGraph/Public/ZoneShapeComponent.h b/Engine/Plugins/Runtime/ZoneGraph/Source/ZoneGraph/Public/ZoneShapeComponent.h
--- a/Engine/Plugins/Runtime/ZoneGraph/Source/ZoneGraph/Public/ZoneShapeComponent.h
+++ b/Engine/Plugins/Runtime/ZoneGraph/Source/ZoneGraph/Public/ZoneShapeComponent.h
@@ -101,6 +101,9 @@ public:
 	/** @return Lane template for each point in the polygon shape. */
 	void GetPolygonLaneProfiles(TArray<FZoneLaneProfile>& OutLaneProfiles) const;
 
+	/** @return Lane profile which is used for spline lanes, and polygon points which are set to Inherit. */
+	const FZoneLaneProfileRef& GetCommonLaneProfile() { return LaneProfile; }
+
 	/** Sets lane profile which is used for spline lanes, and polygon points which are set to Inherit.
 	 * @param LaneProfileRef New lane profile.
 	 */
@@ -127,6 +130,9 @@ public:
 	/** @return View to connected shapes. */
 	TConstArrayView<FZoneShapeConnection> GetConnectedShapes() const { return ConnectedShapes; }
 
+	/** @return Connector at the given point. */
+	const FZoneShapeConnector* GetShapeConnectorByPointIndex(int32 PointIndex) const;
+
 	/** @return Shape type. */
 	UFUNCTION(Category = Zone, BlueprintPure)
 	FZoneShapeType GetShapeType() const { return ShapeType; }
diff --git a/Engine/Plugins/Runtime/ZoneGraph/Source/ZoneGraphEditor/Private/ZoneShapeComponentVisualizer.cpp b/Engine/Plugins/Runtime/ZoneGraph/Source/ZoneGraphEditor/Private/ZoneShapeComponentVisualizer.cpp
--- a/Engine/Plugins/Runtime/ZoneGraph/Source/ZoneGraphEditor/Private/ZoneShapeComponentVisualizer.cpp
+++ b/Engine/Plugins/Runtime/ZoneGraph/Source/ZoneGraphEditor/Private/ZoneShapeComponentVisualizer.cpp
@@ -18,6 +18,7 @@
 #include "LevelEditorActions.h"
 #include "ScopedTransaction.h"
 #include "ActorEditorUtils.h"
+#include "ZoneGraphQuery.h"
 #include "ZoneGraphSubsystem.h"
 #include "ZoneGraphSettings.h"
 #include "ZoneShapeActor.h"
@@ -25,6 +26,8 @@
 #include "ZoneShapeUtilities.h"
 #include "ZoneGraphRenderingUtilities.h"
 #include "BezierUtilities.h"
+#include "CanvasTypes.h"
+#include "SceneManagement.h"
 
 // Uncomment to draw additional rotation debug visualizations.
 // #define ZONEGRAPH_DEBUG_ROTATIONS
@@ -413,9 +416,70 @@ void FZoneShapeComponentVisualizer::DrawVisualization(const UActorComponent* Com
 #endif
 	}
 
+	// Draw auto connection range indicator
+	if (bIsActiveComponent && bIsAutoConnecting && ShapePoints.IsValidIndex(SelectedPointForConnecting))
+	{
+		// Draw a wire sphere
+		const FZoneShapePoint& DraggedPoint = ShapePoints[SelectedPointForConnecting];
+		FVector Center = ShapeComp->GetComponentTransform().TransformPosition(DraggedPoint.Position);
+		const FTransform Transform(FQuat::Identity, Center);
+		constexpr FColor IndicatorColor = FColor(255, 165, 0, 255);
+		const UZoneGraphSettings* ZoneGraphSettings = GetDefault<UZoneGraphSettings>();
+		check(ZoneGraphSettings);
+		const float Radius = ZoneGraphSettings->GetBuildSettings().DragEndpointAutoConnectRange;
+		DrawWireSphere(PDI, Transform, IndicatorColor, Radius, 12, SDPG_World, 0.0f, 0.001f, false);
+
+		// Tint the chevron of the candidate connectors
+		for (int32 i = 0; i < DestShapeConnectorInfos.Num(); i++)
+		{
+			const ZoneShapeConnectorRenderInfo& Info = DestShapeConnectorInfos[i];
+			const FVector WorldPosition = Info.Position;
+			const FVector WorldNormal = Info.Normal;
+			const FVector WorldUp = Info.Up;
+			const FVector WorldSide = FVector::CrossProduct(Info.Normal, Info.Up);
+
+			constexpr FColor GreenColor = FColor(0, 255, 0, 255);
+			constexpr FColor YellowColor = FColor(255, 255, 0, 255);
+			const FColor& ChevronColor = i == ClosestShapeConnectorInfoIndex ? GreenColor : YellowColor;
+			PDI->DrawLine(WorldPosition - WorldNormal * 20, WorldPosition - WorldSide * 20, ChevronColor, SDPG_World, 4, DepthBias, true);
+			PDI->DrawLine(WorldPosition - WorldNormal * 20, WorldPosition + WorldSide * 20, ChevronColor, SDPG_World, 4, DepthBias, true);
+		}
+	}
+
 	PDI->SetHitProxy(nullptr);
 }
 
+void FZoneShapeComponentVisualizer::DrawVisualizationHUD(const UActorComponent* Component, const FViewport* Viewport, const FSceneView* View, FCanvas* Canvas)
+{
+	const UZoneShapeComponent* ShapeComp = Cast<const UZoneShapeComponent>(Component);
+	{
+		if (ShapeComp == GetEditedComponent())
+		{
+			check(SelectionState)
+			int32 SelectedControlPoint = SelectionState->GetSelectedControlPoint();
+			int32 LastPointIndexSelected = SelectionState->GetLastPointIndexSelected();
+			if (SelectionState->GetSelectedPoints().Num() == 1 &&
+				(LastPointIndexSelected == 0 || LastPointIndexSelected == (ShapeComp->GetNumPoints() - 1)))
+			{
+				const FIntRect CanvasRect = Canvas->GetViewRect();
+
+				static const FText AutoConnectionHelp = LOCTEXT("ZoneShapeAutoConnectionMessage", "Auto Zone Shape Connection: Hold C and drag zone shape end point close to another shape connector to connect.");
+
+				auto DisplaySnapToActorHelpText = [&](const FText& SnapHelpText)
+				{
+					int32 XL;
+					int32 YL;
+					StringSize(GEngine->GetLargeFont(), XL, YL, *SnapHelpText.ToString());
+					const float DrawPositionX = FMath::FloorToFloat(CanvasRect.Min.X + (CanvasRect.Width() - XL) * 0.5f);
+					const float DrawPositionY = CanvasRect.Min.Y + 50.0f;
+					Canvas->DrawShadowedString(DrawPositionX, DrawPositionY, *SnapHelpText.ToString(), GEngine->GetLargeFont(), FLinearColor::Yellow);
+				};
+				DisplaySnapToActorHelpText(AutoConnectionHelp);
+			}
+		}
+	}
+}
+
 void FZoneShapeComponentVisualizer::ChangeSelectionState(int32 Index, bool bIsCtrlHeld) const
 {
 	check(SelectionState);
@@ -765,12 +829,102 @@ bool FZoneShapeComponentVisualizer::HandleInputDelta(FEditorViewportClient* View
 			return false;
 		}
 
+		int32 SelectedControlPoint = SelectionState->GetSelectedControlPoint();
+		int32 LastPointIndexSelected = SelectionState->GetLastPointIndexSelected();
 		if (SelectionState->GetSelectedControlPoint() != INDEX_NONE)
 		{
 			return TransformSelectedControlPoint(DeltaTranslate);
 		}
 		else if (SelectionState->GetSelectedPoints().Num() > 0)
 		{
+			if (!ViewportClient->IsAltPressed() &&
+				SelectionState->GetSelectedPoints().Num() == 1 &&
+				(LastPointIndexSelected == 0 || LastPointIndexSelected == (ShapeComp->GetNumPoints() - 1)))
+			{
+				// Cache the selected index
+				SelectedPointForConnecting = LastPointIndexSelected;
+				FZoneShapePoint DraggedPoint = ShapeComp->GetPoints()[SelectedPointForConnecting];
+				const FTransform& SourceTransform = ShapeComp->GetComponentTransform();
+				FVector DraggedPointWorldPosition = SourceTransform.TransformPosition(DraggedPoint.Position);
+
+				if (ViewportClient->Viewport->KeyState(EKeys::C))
+				{
+#if WITH_EDITOR
+					bIsAutoConnecting = true;
+
+					DestShapeConnectorInfos.Empty();
+					ClosestShapeConnectorInfoIndex = INDEX_NONE;
+
+					const FZoneShapeConnector* SourceConnector = ShapeComp->GetShapeConnectorByPointIndex(SelectedPointForConnecting);
+
+					UZoneGraphSubsystem* ZoneGraph = UWorld::GetSubsystem<UZoneGraphSubsystem>(ShapeComp->GetWorld());
+					if (SourceConnector && ZoneGraph)
+					{
+						const FVector SourceWorldPosition = SourceTransform.TransformPosition(SourceConnector->Position);
+
+						const UZoneGraphSettings* ZoneGraphSettings = GetDefault<UZoneGraphSettings>();
+						check(ZoneGraphSettings);
+
+						TArray<uint32> QueryResults;
+						const float AutoConnectRange = ZoneGraphSettings->GetBuildSettings().DragEndpointAutoConnectRange;
+						FBox Bounds = FBox::BuildAABB(DraggedPointWorldPosition, FVector(AutoConnectRange));
+						ZoneGraph->GetBuilder().QueryHashGrid(Bounds, QueryResults);
+						const TArray<FZoneGraphBuilderRegisteredComponent>& RegisteredShapeComponents = ZoneGraph->GetBuilder().GetRegisteredZoneShapeComponents();
+						double ShortestDistance = AutoConnectRange;
+						for (uint32 Index : QueryResults)
+						{
+							check(RegisteredShapeComponents.IsValidIndex(int32(Index)));
+							UZoneShapeComponent* DestShapeComp = RegisteredShapeComponents[Index].Component;
+							if (!DestShapeComp || ShapeComp->GetComponentLevel() != DestShapeComp->GetComponentLevel())
+							{
+								continue;
+							}
+
+							const FTransform& DestTransform = DestShapeComp->GetComponentTransform();
+							TConstArrayView<FZoneShapeConnector> DestConnectors = DestShapeComp->GetShapeConnectors();
+
+							for (int32 j = 0; j < DestConnectors.Num(); j++)
+							{
+								const FZoneShapeConnector& DestConnector = DestConnectors[j];
+								const FVector DestWorldPosition = DestTransform.TransformPosition(DestConnector.Position);
+								const FVector DestWorldNormal = DestTransform.TransformVector(DestConnector.Normal);
+
+								double Distance = FVector::Dist(SourceWorldPosition, DestWorldPosition);
+								if (SourceConnector == &DestConnector || SourceConnector->LaneProfile != DestConnector.LaneProfile)
+								{
+									continue;
+								}
+
+								// Check that the profile orientation matches before connecting.
+								if (const FZoneLaneProfile* LaneProfile = ZoneGraphSettings->GetLaneProfileByRef(SourceConnector->LaneProfile))
+								{
+									if (LaneProfile->IsSymmetrical() || SourceConnector->bReverseLaneProfile != DestConnector.bReverseLaneProfile)
+									{
+										if (Distance < AutoConnectRange)
+										{
+											const FVector WorldPosition = DestTransform.TransformPosition(DestConnector.Position);
+											const FVector WorldNormal = DestTransform.TransformVector(DestConnector.Normal);
+											const FVector WorldUp = DestTransform.TransformVector(DestConnector.Up);
+											DestShapeConnectorInfos.Add({ WorldPosition, WorldNormal, WorldUp });
+										}
+
+										if (ShortestDistance > Distance)
+										{
+											ShortestDistance = Distance;
+											ClosestShapeConnectorInfoIndex = DestShapeConnectorInfos.Num() - 1;
+
+											NearestPointWorldPosition = DestWorldPosition;
+											NearestPointWorldNormal = DestWorldNormal;
+										}
+									}
+								}
+							}
+						}
+					}
+#endif
+				}
+			}
+
 			if (ViewportClient->IsAltPressed())
 			{
 				if (ViewportClient->GetWidgetMode() == UE::Widget::WM_Translate && ViewportClient->GetCurrentWidgetAxis() != EAxisList::None)
@@ -929,7 +1083,12 @@ bool FZoneShapeComponentVisualizer::HandleInputKey(FEditorViewportClient* Viewpo
 	bool bHandled = false;
 
 	UZoneShapeComponent* ShapeComp = GetEditedShapeComponent();
-	if (ShapeComp != nullptr && IsAnySelectedPointIndexOutOfRange(*ShapeComp))
+	if (!ShapeComp)
+	{
+		return false;
+	}
+	
+	if (IsAnySelectedPointIndexOutOfRange(*ShapeComp))
 	{
 		// Something external has changed the number of shape points, meaning that the cached selected keys are no longer valid
 		EndEditing();
@@ -947,6 +1106,66 @@ bool FZoneShapeComponentVisualizer::HandleInputKey(FEditorViewportClient* Viewpo
 
 		bHasCachedRotation = false;
 		CachedRotation = FQuat::Identity;
+
+		if (bIsAutoConnecting && SelectedPointForConnecting >= 0 && SelectedPointForConnecting < ShapeComp->GetNumPoints())
+		{
+			const UZoneGraphSettings* ZoneGraphSettings = GetDefault<UZoneGraphSettings>();
+			check(ZoneGraphSettings);
+
+			const FZoneLaneProfile* LaneProfile = ZoneGraphSettings->GetLaneProfileByRef(ShapeComp->GetCommonLaneProfile());
+			check(LaneProfile);
+			float HalfLanesTotalWidth = LaneProfile->GetLanesTotalWidth() * 0.5;
+
+			FZoneShapePoint& DraggedPoint = ShapeComp->GetMutablePoints()[SelectedPointForConnecting];
+
+#if WITH_EDITOR
+			if (const FZoneShapeConnector* SourceConnector = ShapeComp->GetShapeConnectorByPointIndex(SelectedPointForConnecting))
+			{
+				const FTransform& SourceTransform = ShapeComp->GetComponentTransform();
+				const FVector SourceWorldNormal = SourceTransform.TransformVector(SourceConnector->Normal);
+
+				const FZoneGraphBuildSettings& BuildSettings = ZoneGraphSettings->GetBuildSettings();
+				static const float ConnectionSnapAngleCos = FMath::Cos(FMath::DegreesToRadians(BuildSettings.ConnectionSnapAngle));
+
+				if (ClosestShapeConnectorInfoIndex != INDEX_NONE)
+				{
+					// Snap point location
+					DraggedPoint.Position = SourceTransform.InverseTransformPosition(NearestPointWorldPosition);
+					FVector Normal = SourceTransform.InverseTransformVector(NearestPointWorldNormal);
+					const FRotator Rotation = FRotationMatrix::MakeFromX(SelectedPointForConnecting == 0 ? Normal : -Normal).Rotator();
+					DraggedPoint.Rotation = Rotation;
+
+					// If the zone shape is a spline and the point type is not Bezier, setting the point rotation doesn't work.
+					// An extra point is needed to align the connectors and make it connect.
+					if (ShapeComp->GetShapeType() == FZoneShapeType::Spline &&
+						DraggedPoint.Type != FZoneShapePointType::Bezier &&
+						FVector::DotProduct(SourceWorldNormal, -NearestPointWorldNormal) <= ConnectionSnapAngleCos)
+					{
+						// Add extra point
+						TArray<FZoneShapePoint>& Points = ShapeComp->GetMutablePoints();
+						FZoneShapePoint ExtraPoint = DraggedPoint;
+						ExtraPoint.Position += Normal * HalfLanesTotalWidth;
+						ExtraPoint.Rotation = Rotation;
+						Points.Insert(ExtraPoint, ShapeComp->GetNumPoints() - 1);
+					}
+
+					// Update shape
+					ShapeComp->UpdateShape();
+				}
+			}
+#endif
+		}
+
+		bIsAutoConnecting = false;
+		DestShapeConnectorInfos.Empty();
+		ClosestShapeConnectorInfoIndex = INDEX_NONE;
+	}
+
+	if (Key == EKeys::C && Event == IE_Released)
+	{
+		bIsAutoConnecting = false;
+		DestShapeConnectorInfos.Empty();
+		ClosestShapeConnectorInfoIndex = INDEX_NONE;
 	}
 
 	if (Key == EKeys::LeftMouseButton && Event == IE_Pressed)
@@ -1706,7 +1925,7 @@ void FZoneShapeComponentVisualizer::BreakAtPoint(bool bCreateNewActor) const
 			NewShapeComponent->SetWorldTransform(ShapeComp->GetComponentTransform());
 			ShapeOwner->AddInstanceComponent(NewShapeComponent);
 			NewShapeComponent->RegisterComponent();
-			NewShapeComponent->AttachToComponent(ShapeComp, FAttachmentTransformRules::KeepWorldTransform); // Should we attach to the root component?
+			NewShapeComponent->AttachToComponent(ShapeComp, FAttachmentTransformRules::KeepWorldTransform);
 			NewShapeComponent->Modify();
 		}
 
diff --git a/Engine/Plugins/Runtime/ZoneGraph/Source/ZoneGraphEditor/Private/ZoneShapeComponentVisualizer.h b/Engine/Plugins/Runtime/ZoneGraph/Source/ZoneGraphEditor/Private/ZoneShapeComponentVisualizer.h
--- a/Engine/Plugins/Runtime/ZoneGraph/Source/ZoneGraphEditor/Private/ZoneShapeComponentVisualizer.h
+++ b/Engine/Plugins/Runtime/ZoneGraph/Source/ZoneGraphEditor/Private/ZoneShapeComponentVisualizer.h
@@ -162,6 +162,7 @@ public:
 	virtual void OnRegister() override;
 	virtual void DrawVisualization(const UActorComponent* Component, const FSceneView* View, FPrimitiveDrawInterface* PDI) override;
 	virtual bool VisProxyHandleClick(FEditorViewportClient* InViewportClient, HComponentVisProxy* VisProxy, const FViewportClick& Click) override;
+	virtual void DrawVisualizationHUD(const UActorComponent* Component, const FViewport* Viewport, const FSceneView* View, FCanvas* Canvas) override;
 	virtual void EndEditing() override;
 	virtual bool GetWidgetLocation(const FEditorViewportClient* ViewportClient, FVector& OutLocation) const override;
 	virtual bool GetCustomInputCoordinateSystem(const FEditorViewportClient* ViewportClient, FMatrix& OutMatrix) const override;
@@ -192,6 +193,26 @@ public:
 
 protected:
 
+	struct ZoneShapeConnectorRenderInfo
+	{
+		ZoneShapeConnectorRenderInfo(FVector InPosition, FVector InNormal, FVector InUp)
+			: Position(InPosition)
+			, Normal(InNormal)
+			, Up(InUp)
+		{
+
+		}
+
+		// Position of the connector.
+		FVector Position = FVector::ZeroVector;
+
+		// Normal direction of the connector.
+		FVector Normal = FVector::ForwardVector;
+
+		// Up direction of the connector.
+		FVector Up = FVector::UpVector;
+	};
+
 	/** Determine if any selected key index is out of range (perhaps because something external has modified the shape) */
 	bool IsAnySelectedPointIndexOutOfRange(const UZoneShapeComponent& Comp) const;
 
@@ -289,7 +310,15 @@ protected:
 
 	bool bIsSelectingComponent = false;
 
+	bool bIsAutoConnecting = false;
+	int32 SelectedPointForConnecting = -1;
+	TArray<ZoneShapeConnectorRenderInfo> DestShapeConnectorInfos;
+	int32 ClosestShapeConnectorInfoIndex = INDEX_NONE;
+	FVector NearestPointWorldPosition;
+	FVector NearestPointWorldNormal;
+
 private:
+
 	void BreakAtPoint(bool bCreateNewActor) const;
 	void BreakAtSegment(bool bCreateNewActor) const;
 };
