ZoneGraph: crash fix (wrong condition order)

diff --git a/Engine/Plugins/Runtime/ZoneGraph/Source/ZoneGraph/Public/ZoneGraphSubsystem.h b/Engine/Plugins/Runtime/ZoneGraph/Source/ZoneGraph/Public/ZoneGraphSubsystem.h
--- a/Engine/Plugins/Runtime/ZoneGraph/Source/ZoneGraph/Public/ZoneGraphSubsystem.h
+++ b/Engine/Plugins/Runtime/ZoneGraph/Source/ZoneGraph/Public/ZoneGraphSubsystem.h
@@ -58,7 +58,7 @@ public:
 	// Returns Zone Graph data storage associated with specified handle, or nullptr if not found.
 	const FZoneGraphStorage* GetZoneGraphStorage(const FZoneGraphDataHandle DataHandle) const
 	{
-		if (int32(DataHandle.Generation) == RegisteredZoneGraphData[DataHandle.Index].Generation && int32(DataHandle.Index) < RegisteredZoneGraphData.Num())
+		if (int32(DataHandle.Index) < RegisteredZoneGraphData.Num() && int32(DataHandle.Generation) == RegisteredZoneGraphData[DataHandle.Index].Generation)
 		{
 			if (const AZoneGraphData* Data = RegisteredZoneGraphData[DataHandle.Index].ZoneGraphData)
 			{
