[ZoneGraph] When dragging an end point onto another zone shape, automatically create an intersection.

diff --git a/Engine/Plugins/Runtime/ZoneGraph/Source/ZoneGraph/Public/ZoneGraphTypes.h b/Engine/Plugins/Runtime/ZoneGraph/Source/ZoneGraph/Public/ZoneGraphTypes.h
--- a/Engine/Plugins/Runtime/ZoneGraph/Source/ZoneGraph/Public/ZoneGraphTypes.h
+++ b/Engine/Plugins/Runtime/ZoneGraph/Source/ZoneGraph/Public/ZoneGraphTypes.h
@@ -998,5 +998,17 @@ struct ZONEGRAPH_API FZoneGraphBuildSettings
 
 	/** Max distance to auto connect a dragged end control point to the closest connector when released. */
 	UPROPERTY(Category = PointSnapping, EditAnywhere)
-	float DragEndpointAutoConnectRange = 250.0f;
+	double DragEndpointAutoConnectRange = 250.0;
+
+	/** Max distance to automatically create an intersection with a dragged end control point and the overlapping zone shape when released. */
+	UPROPERTY(Category = PointSnapping, EditAnywhere)
+	double DragEndpointAutoIntersectionRange = 250.0;
+
+	/** If the location to create auto intersection is closer than the tolerance to a spline point, use the spline point instead of nearest point on spline to create the intersection. */
+	UPROPERTY(Category = PointSnapping, EditAnywhere)
+	double SnapAutoIntersectionToClosestPointTolerance = 100.0;
+
+	/** If true, draws a sphere for the auto connection/interseciont area. Otherwise, draws a circle. */
+	UPROPERTY(Category = PointSnapping, EditAnywhere)
+	bool bShow3DRadiusForAutoConnectionAndIntersection = false;
 };
diff --git a/Engine/Plugins/Runtime/ZoneGraph/Source/ZoneGraphEditor/Private/ZoneShapeComponentVisualizer.cpp b/Engine/Plugins/Runtime/ZoneGraph/Source/ZoneGraphEditor/Private/ZoneShapeComponentVisualizer.cpp
--- a/Engine/Plugins/Runtime/ZoneGraph/Source/ZoneGraphEditor/Private/ZoneShapeComponentVisualizer.cpp
+++ b/Engine/Plugins/Runtime/ZoneGraph/Source/ZoneGraphEditor/Private/ZoneShapeComponentVisualizer.cpp
@@ -40,6 +40,110 @@ IMPLEMENT_HIT_PROXY(HZoneShapeControlPointProxy, HZoneShapeVisProxy);
 #define LOCTEXT_NAMESPACE "ZoneShapeComponentVisualizer"
 DEFINE_LOG_CATEGORY_STATIC(LogZoneShapeComponentVisualizer, Log, All)
 
+
+namespace UE::ZoneGraph::Editor::Private
+{
+	double GetClockwiseAngle(const FVector& P)
+	{
+		return -FMath::Atan2(P.X, -P.Y);
+	}
+
+	bool ComparePoints(const FVector& P1, const FVector& P2)
+	{
+		return GetClockwiseAngle(P1) > GetClockwiseAngle(P2);
+	}
+
+	void SortPolygonPointsCounterclockwise(UZoneShapeComponent* PolygonShapeComp)
+	{
+		if (PolygonShapeComp->GetShapeType() != FZoneShapeType::Polygon)
+		{
+			return;
+		}
+
+		TArray<FZoneShapePoint>& Points = PolygonShapeComp->GetMutablePoints();
+
+		// Compute the center
+		FVector Center = FVector::ZeroVector;
+		for (const FZoneShapePoint& Point : Points)
+		{
+			Center += Point.Position;
+		}
+		Center /= Points.Num();
+
+		Points.Sort([Center](const FZoneShapePoint& Point1, const FZoneShapePoint& Point2) {
+			return ComparePoints(Point1.Position - Center, Point2.Position - Center);
+		});
+	}
+
+	FVector GetPositionOnSegment(const TArray<FZoneShapePoint>& Points, int32 SegmentIndex, float SegmentT)
+	{
+		const int32 NumPoints = Points.Num();
+		const int32 StartPointIdx = SegmentIndex;
+		const int32 EndPointIdx = (SegmentIndex + 1) % NumPoints;
+		const FZoneShapePoint& StartPoint = Points[StartPointIdx];
+		const FZoneShapePoint& EndPoint = Points[EndPointIdx];
+
+		FVector StartPosition(ForceInitToZero), StartControlPoint(ForceInitToZero), EndControlPoint(ForceInitToZero), EndPosition(ForceInitToZero);
+		UE::ZoneShape::Utilities::GetCubicBezierPointsFromShapeSegment(StartPoint, EndPoint, FMatrix::Identity, StartPosition, StartControlPoint, EndControlPoint, EndPosition);
+
+		return UE::CubicBezier::Eval(StartPosition, StartControlPoint, EndControlPoint, EndPosition, SegmentT);
+	}
+
+	void SetPolygonPointLaneProfileToMatchSpline(FZoneShapePoint& Point, UZoneShapeComponent* Polygon, UZoneShapeComponent* Spline)
+	{
+		Point.Type = FZoneShapePointType::LaneProfile;
+		const FZoneLaneProfileRef ShapeComponent0LaneProfileRef = Spline->GetCommonLaneProfile();
+		const int32 ProfileIndex = Polygon->AddUniquePerPointLaneProfile(ShapeComponent0LaneProfileRef);
+		if (ProfileIndex != INDEX_NONE)
+		{
+			Point.LaneProfile = (uint8)ProfileIndex;
+		}
+	}
+
+	void SetPointPositionRotation(
+		FZoneShapePoint& Point,
+		const FTransform& SourceTransform,
+		const FVector& TargetPointWorldPosition,
+		const FVector& TargetPointWorldNormal)
+	{
+		Point.Position = SourceTransform.InverseTransformPosition(TargetPointWorldPosition);
+		FVector Normal = SourceTransform.InverseTransformVector(TargetPointWorldNormal);
+		Point.Rotation = FRotationMatrix::MakeFromX(Normal).Rotator();
+	}
+
+	void SnapConnect(
+		UZoneShapeComponent* ShapeComp,
+		FZoneShapePoint& DraggedPoint,
+		const FTransform& SourceTransform,
+		const FVector& SourceWorldNormal,
+		const FVector& TargetPointWorldPosition,
+		const FVector& TargetPointWorldNormal,
+		double ConnectionSnapAngleCos,
+		double HalfLanesTotalWidth)
+	{
+		// Snap point location
+		UE::ZoneGraph::Editor::Private::SetPointPositionRotation(DraggedPoint, SourceTransform, TargetPointWorldPosition, TargetPointWorldNormal);
+
+		// If the zone shape is a spline and the point type is not Bezier, setting the point rotation doesn't work.
+		// An extra point is needed to align the connectors and make it connect.
+		if (ShapeComp->GetShapeType() == FZoneShapeType::Spline &&
+			DraggedPoint.Type != FZoneShapePointType::Bezier &&
+			FVector::DotProduct(SourceWorldNormal, -TargetPointWorldNormal) <= ConnectionSnapAngleCos)
+		{
+			// Add extra point
+			TArray<FZoneShapePoint>& Points = ShapeComp->GetMutablePoints();
+			FZoneShapePoint ExtraPoint = DraggedPoint;
+			ExtraPoint.Position += SourceTransform.InverseTransformVector(TargetPointWorldNormal) * HalfLanesTotalWidth;
+			ExtraPoint.Rotation = DraggedPoint.Rotation;
+			Points.Insert(ExtraPoint, ShapeComp->GetNumPoints() - 1);
+		}
+
+		// Update shape
+		ShapeComp->UpdateShape();
+	}
+} // UE::ZoneGraph::Editor::Private
+
+
 /** Define commands for the shape component visualizer */
 class FZoneShapeComponentVisualizerCommands : public TCommands<FZoneShapeComponentVisualizerCommands>
 {
@@ -216,14 +320,20 @@ void FZoneShapeComponentVisualizer::DrawVisualization(const UActorComponent* Com
 	{
 		return;
 	}
+	const UZoneGraphSettings* ZoneGraphSettings = GetDefault<UZoneGraphSettings>();
+	if (!ZoneGraphSettings)
+	{
+		return;
+	}
+
+	const FZoneGraphBuildSettings& BuildSettings = ZoneGraphSettings->GetBuildSettings();
+
 	const FMatrix LocalToWorld = ShapeComp->GetComponentTransform().ToMatrixWithScale();
 
 	// Distance culling.
 	float ShapeMaxDrawDistance = MAX_flt;
-	if (const UZoneGraphSettings* ZoneGraphSettings = GetDefault<UZoneGraphSettings>())
-	{
-		ShapeMaxDrawDistance = ZoneGraphSettings->GetShapeMaxDrawDistance();
-	}
+	ShapeMaxDrawDistance = ZoneGraphSettings->GetShapeMaxDrawDistance();
+	
 	const float MaxDrawDistanceSqr = FMath::Square(ShapeMaxDrawDistance);
 
 	// Taking into account the min and maximum drawing distance
@@ -305,8 +415,8 @@ void FZoneShapeComponentVisualizer::DrawVisualization(const UActorComponent* Com
 	if (ShapePoints.Num() > 1)
 	{
 		const int32 NumPoints = ShapePoints.Num();
-		int StartIdx = ShapeComp->IsShapeClosed() ? (NumPoints - 1) : 0;
-		int Idx = ShapeComp->IsShapeClosed() ? 0 : 1;
+		int32 StartIdx = ShapeComp->IsShapeClosed() ? (NumPoints - 1) : 0;
+		int32 Idx = ShapeComp->IsShapeClosed() ? 0 : 1;
 
 		TArray<FVector> CurvePoints;
 
@@ -416,33 +526,81 @@ void FZoneShapeComponentVisualizer::DrawVisualization(const UActorComponent* Com
 #endif
 	}
 
-	// Draw auto connection range indicator
-	if (bIsActiveComponent && bIsAutoConnecting && ShapePoints.IsValidIndex(SelectedPointForConnecting))
+	if (bIsActiveComponent && (bIsAutoConnecting || bIsCreatingIntersection) && ShapePoints.IsValidIndex(SelectedPointForConnecting))
 	{
-		// Draw a wire sphere
 		const FZoneShapePoint& DraggedPoint = ShapePoints[SelectedPointForConnecting];
-		FVector Center = ShapeComp->GetComponentTransform().TransformPosition(DraggedPoint.Position);
+		const FVector Center = ShapeComp->GetComponentTransform().TransformPosition(DraggedPoint.Position);
 		const FTransform Transform(FQuat::Identity, Center);
-		constexpr FColor IndicatorColor = FColor(255, 165, 0, 255);
-		const UZoneGraphSettings* ZoneGraphSettings = GetDefault<UZoneGraphSettings>();
-		check(ZoneGraphSettings);
-		const float Radius = ZoneGraphSettings->GetBuildSettings().DragEndpointAutoConnectRange;
-		DrawWireSphere(PDI, Transform, IndicatorColor, Radius, 12, SDPG_World, 0.0f, 0.001f, false);
+		constexpr FColor IndicatorColor = FColor(255, 192, 32, 255);
+		constexpr FColor InnerIndicatorColor = FColor(192, 128, 16, 255);
+
+		double IndicatorRadius = 0.0;
+		double IndicatorInnerRadius = 0.0;
 
-		// Tint the chevron of the candidate connectors
-		for (int32 i = 0; i < DestShapeConnectorInfos.Num(); i++)
+		if (bIsCreatingIntersection)
 		{
-			const ZoneShapeConnectorRenderInfo& Info = DestShapeConnectorInfos[i];
-			const FVector WorldPosition = Info.Position;
-			const FVector WorldNormal = Info.Normal;
-			const FVector WorldUp = Info.Up;
-			const FVector WorldSide = FVector::CrossProduct(Info.Normal, Info.Up);
+			if (UZoneShapeComponent* TargetShapeComponent = CreateIntersectionState.WeakTargetShapeComponent.Get())
+			{
+				const FTransform& TargetShapeTransform = TargetShapeComponent->GetComponentTransform();
+
+				// Draw X at the indicative location where the intersection will be build.
+				constexpr double MarkerHalfSize = 10.0;
+				const FVector AxisX = TargetShapeTransform.GetUnitAxis(EAxis::X);
+				const FVector AxisY = TargetShapeTransform.GetUnitAxis(EAxis::Y);
+
+				PDI->DrawLine(
+					CreateIntersectionState.PreviewLocation - AxisX * MarkerHalfSize - AxisY * MarkerHalfSize,
+					CreateIntersectionState.PreviewLocation + AxisX * MarkerHalfSize + AxisY * MarkerHalfSize,
+					FColor::Red, SDPG_World, 4.0f);
+				PDI->DrawLine(
+					CreateIntersectionState.PreviewLocation - AxisX * MarkerHalfSize + AxisY * MarkerHalfSize,
+					CreateIntersectionState.PreviewLocation + AxisX * MarkerHalfSize - AxisY * MarkerHalfSize,
+					FColor::Red, SDPG_World, 4.0f);
+			}
+
+			IndicatorRadius = BuildSettings.DragEndpointAutoIntersectionRange;
+			IndicatorInnerRadius = BuildSettings.SnapAutoIntersectionToClosestPointTolerance;
+		}
 
-			constexpr FColor GreenColor = FColor(0, 255, 0, 255);
-			constexpr FColor YellowColor = FColor(255, 255, 0, 255);
-			const FColor& ChevronColor = i == ClosestShapeConnectorInfoIndex ? GreenColor : YellowColor;
-			PDI->DrawLine(WorldPosition - WorldNormal * 20, WorldPosition - WorldSide * 20, ChevronColor, SDPG_World, 4, DepthBias, true);
-			PDI->DrawLine(WorldPosition - WorldNormal * 20, WorldPosition + WorldSide * 20, ChevronColor, SDPG_World, 4, DepthBias, true);
+		if (bIsAutoConnecting)
+		{
+			for (int32 Index = 0; Index < AutoConnectState.DestShapeConnectorInfos.Num(); Index++)
+			{
+				const bool bIsClosest = (Index == AutoConnectState.ClosestShapeConnectorInfoIndex);
+				
+				// Draw a square at the potential snap position
+				const ZoneShapeConnectorRenderInfo& Info = AutoConnectState.DestShapeConnectorInfos[Index];
+				const FColor& ChevronColor = bIsClosest ? FColor::Red : FColor::Silver;
+				const FVector AxisX = Info.Foward.RotateAngleAxis(-45, Info.Up);
+				const FVector AxisY = Info.Foward.RotateAngleAxis(45, Info.Up);
+				DrawRectangle(PDI, Info.Position, AxisX, AxisY, ChevronColor, 20.f, 20.f, SDPG_World, 4.f);
+			}
+			IndicatorRadius = BuildSettings.DragEndpointAutoConnectRange;
+		}
+
+		// Draw auto connection/intersection range indicator
+		if (IndicatorRadius > 0.0)
+		{
+			if (BuildSettings.bShow3DRadiusForAutoConnectionAndIntersection)
+			{
+				DrawWireSphere(PDI, Transform, IndicatorColor, IndicatorRadius, 32, SDPG_World, 0.0f, 0.001f, false);
+			}
+			else
+			{
+				DrawCircle(PDI, Center, FVector::XAxisVector, FVector::YAxisVector, IndicatorColor, IndicatorRadius, 32, SDPG_World);
+			}
+		}
+		
+		if (IndicatorInnerRadius > 0.0)
+		{
+			if (BuildSettings.bShow3DRadiusForAutoConnectionAndIntersection)
+			{
+				DrawWireSphere(PDI, Transform, InnerIndicatorColor, IndicatorInnerRadius, 24, SDPG_World, 0.0f, 0.001f, false);
+			}
+			else
+			{
+				DrawCircle(PDI, Center, FVector::XAxisVector, FVector::YAxisVector, InnerIndicatorColor, IndicatorInnerRadius, 24, SDPG_World);
+			}
 		}
 	}
 
@@ -453,7 +611,7 @@ void FZoneShapeComponentVisualizer::DrawVisualizationHUD(const UActorComponent*
 {
 	const UZoneShapeComponent* ShapeComp = Cast<const UZoneShapeComponent>(Component);
 	{
-		if (ShapeComp == GetEditedComponent())
+		if (ShapeComp != nullptr && ShapeComp == GetEditedComponent())
 		{
 			check(SelectionState)
 			int32 SelectedControlPoint = SelectionState->GetSelectedControlPoint();
@@ -464,17 +622,25 @@ void FZoneShapeComponentVisualizer::DrawVisualizationHUD(const UActorComponent*
 				const FIntRect CanvasRect = Canvas->GetViewRect();
 
 				static const FText AutoConnectionHelp = LOCTEXT("ZoneShapeAutoConnectionMessage", "Auto Zone Shape Connection: Hold C and drag zone shape end point close to another shape connector to connect.");
+				static const FText AutoIntersectionHelp = LOCTEXT("ZoneShapeAutoIntersectionMessage", "Auto Zone Shape Intersection: Hold X and drag zone shape end point close to another shape to create an intersection.");
 
-				auto DisplaySnapToActorHelpText = [&](const FText& SnapHelpText)
+				auto DisplaySnapToActorHelpText = [&](const FText& SnapHelpText, double YOffset)
 				{
 					int32 XL;
 					int32 YL;
 					StringSize(GEngine->GetLargeFont(), XL, YL, *SnapHelpText.ToString());
-					const float DrawPositionX = FMath::FloorToFloat(CanvasRect.Min.X + (CanvasRect.Width() - XL) * 0.5f);
-					const float DrawPositionY = CanvasRect.Min.Y + 50.0f;
+					const double DrawPositionX = FMath::FloorToDouble(CanvasRect.Min.X + (CanvasRect.Width() - XL) * 0.5);
+					const double DrawPositionY = CanvasRect.Min.Y + 50.0 + YOffset;
 					Canvas->DrawShadowedString(DrawPositionX, DrawPositionY, *SnapHelpText.ToString(), GEngine->GetLargeFont(), FLinearColor::Yellow);
 				};
-				DisplaySnapToActorHelpText(AutoConnectionHelp);
+				if (CanAutoConnect(ShapeComp))
+				{
+					DisplaySnapToActorHelpText(AutoConnectionHelp, 0.0);
+				}
+				if (CanAutoCreateIntersection(ShapeComp))
+				{
+					DisplaySnapToActorHelpText(AutoIntersectionHelp, 20.0);
+				}
 			}
 		}
 	}
@@ -829,8 +995,7 @@ bool FZoneShapeComponentVisualizer::HandleInputDelta(FEditorViewportClient* View
 			return false;
 		}
 
-		int32 SelectedControlPoint = SelectionState->GetSelectedControlPoint();
-		int32 LastPointIndexSelected = SelectionState->GetLastPointIndexSelected();
+		const int32 LastPointIndexSelected = SelectionState->GetLastPointIndexSelected();
 		if (SelectionState->GetSelectedControlPoint() != INDEX_NONE)
 		{
 			return TransformSelectedControlPoint(DeltaTranslate);
@@ -843,86 +1008,18 @@ bool FZoneShapeComponentVisualizer::HandleInputDelta(FEditorViewportClient* View
 			{
 				// Cache the selected index
 				SelectedPointForConnecting = LastPointIndexSelected;
-				FZoneShapePoint DraggedPoint = ShapeComp->GetPoints()[SelectedPointForConnecting];
-				const FTransform& SourceTransform = ShapeComp->GetComponentTransform();
-				FVector DraggedPointWorldPosition = SourceTransform.TransformPosition(DraggedPoint.Position);
+				const FZoneShapePoint& DraggedPoint = ShapeComp->GetPoints()[SelectedPointForConnecting];
 
+#if WITH_EDITOR
 				if (ViewportClient->Viewport->KeyState(EKeys::C))
 				{
-#if WITH_EDITOR
-					bIsAutoConnecting = true;
-
-					DestShapeConnectorInfos.Empty();
-					ClosestShapeConnectorInfoIndex = INDEX_NONE;
-
-					const FZoneShapeConnector* SourceConnector = ShapeComp->GetShapeConnectorByPointIndex(SelectedPointForConnecting);
-
-					UZoneGraphSubsystem* ZoneGraph = UWorld::GetSubsystem<UZoneGraphSubsystem>(ShapeComp->GetWorld());
-					if (SourceConnector && ZoneGraph)
-					{
-						const FVector SourceWorldPosition = SourceTransform.TransformPosition(SourceConnector->Position);
-
-						const UZoneGraphSettings* ZoneGraphSettings = GetDefault<UZoneGraphSettings>();
-						check(ZoneGraphSettings);
-
-						TArray<uint32> QueryResults;
-						const float AutoConnectRange = ZoneGraphSettings->GetBuildSettings().DragEndpointAutoConnectRange;
-						FBox Bounds = FBox::BuildAABB(DraggedPointWorldPosition, FVector(AutoConnectRange));
-						ZoneGraph->GetBuilder().QueryHashGrid(Bounds, QueryResults);
-						const TArray<FZoneGraphBuilderRegisteredComponent>& RegisteredShapeComponents = ZoneGraph->GetBuilder().GetRegisteredZoneShapeComponents();
-						double ShortestDistance = AutoConnectRange;
-						for (uint32 Index : QueryResults)
-						{
-							check(RegisteredShapeComponents.IsValidIndex(int32(Index)));
-							UZoneShapeComponent* DestShapeComp = RegisteredShapeComponents[Index].Component;
-							if (!DestShapeComp || ShapeComp->GetComponentLevel() != DestShapeComp->GetComponentLevel())
-							{
-								continue;
-							}
-
-							const FTransform& DestTransform = DestShapeComp->GetComponentTransform();
-							TConstArrayView<FZoneShapeConnector> DestConnectors = DestShapeComp->GetShapeConnectors();
-
-							for (int32 j = 0; j < DestConnectors.Num(); j++)
-							{
-								const FZoneShapeConnector& DestConnector = DestConnectors[j];
-								const FVector DestWorldPosition = DestTransform.TransformPosition(DestConnector.Position);
-								const FVector DestWorldNormal = DestTransform.TransformVector(DestConnector.Normal);
-
-								double Distance = FVector::Dist(SourceWorldPosition, DestWorldPosition);
-								if (SourceConnector == &DestConnector || SourceConnector->LaneProfile != DestConnector.LaneProfile)
-								{
-									continue;
-								}
-
-								// Check that the profile orientation matches before connecting.
-								if (const FZoneLaneProfile* LaneProfile = ZoneGraphSettings->GetLaneProfileByRef(SourceConnector->LaneProfile))
-								{
-									if (LaneProfile->IsSymmetrical() || SourceConnector->bReverseLaneProfile != DestConnector.bReverseLaneProfile)
-									{
-										if (Distance < AutoConnectRange)
-										{
-											const FVector WorldPosition = DestTransform.TransformPosition(DestConnector.Position);
-											const FVector WorldNormal = DestTransform.TransformVector(DestConnector.Normal);
-											const FVector WorldUp = DestTransform.TransformVector(DestConnector.Up);
-											DestShapeConnectorInfos.Add({ WorldPosition, WorldNormal, WorldUp });
-										}
-
-										if (ShortestDistance > Distance)
-										{
-											ShortestDistance = Distance;
-											ClosestShapeConnectorInfoIndex = DestShapeConnectorInfos.Num() - 1;
-
-											NearestPointWorldPosition = DestWorldPosition;
-											NearestPointWorldNormal = DestWorldNormal;
-										}
-									}
-								}
-							}
-						}
-					}
-#endif
+					DetectCloseByShapeForAutoConnection(ShapeComp, DraggedPoint);
+				}
+				else if (ViewportClient->Viewport->KeyState(EKeys::X) && CanAutoCreateIntersection(ShapeComp))
+				{
+					DetectCloseByShapeForAutoIntersectionCreation(ShapeComp, DraggedPoint);
 				}
+#endif
 			}
 
 			if (ViewportClient->IsAltPressed())
@@ -959,6 +1056,280 @@ bool FZoneShapeComponentVisualizer::HandleInputDelta(FEditorViewportClient* View
 	return false;
 }
 
+void FZoneShapeComponentVisualizer::DetectCloseByShapeForAutoConnection(const UZoneShapeComponent* ShapeComp, const FZoneShapePoint& DraggedPoint)
+{
+	ClearAutoConnectingStatus();
+	bIsAutoConnecting = true;
+
+	const UZoneGraphSettings* ZoneGraphSettings = GetDefault<UZoneGraphSettings>();
+	if (!ZoneGraphSettings)
+	{
+		return;
+	}
+	UZoneGraphSubsystem* ZoneGraph = UWorld::GetSubsystem<UZoneGraphSubsystem>(ShapeComp->GetWorld());
+	if (!ZoneGraph)
+	{
+		return;
+	}
+
+	const FZoneShapeConnector* SourceConnector = ShapeComp->GetShapeConnectorByPointIndex(SelectedPointForConnecting);
+	if (!SourceConnector)
+	{
+		return;
+	}
+
+	const FZoneGraphBuildSettings& BuildSettings = ZoneGraphSettings->GetBuildSettings();
+	const TArray<FZoneGraphBuilderRegisteredComponent>& RegisteredShapeComponents = ZoneGraph->GetBuilder().GetRegisteredZoneShapeComponents();
+
+	const FTransform& SourceTransform = ShapeComp->GetComponentTransform();
+	const FVector SourceWorldPosition = SourceTransform.TransformPosition(SourceConnector->Position);
+	const FVector DraggedPointWorldPosition = SourceTransform.TransformPosition(DraggedPoint.Position);
+	
+	TArray<uint32> QueryResults;
+	const FBox Bounds = FBox::BuildAABB(DraggedPointWorldPosition, FVector(BuildSettings.DragEndpointAutoConnectRange));
+	ZoneGraph->GetBuilder().QueryHashGrid(Bounds, QueryResults);
+		
+	double ShortestDistance = BuildSettings.DragEndpointAutoConnectRange;
+	for (const uint32 ComponentIndex : QueryResults)
+	{
+		check(RegisteredShapeComponents.IsValidIndex(int32(ComponentIndex)));
+		const UZoneShapeComponent* DestShapeComp = RegisteredShapeComponents[ComponentIndex].Component;
+		if (!DestShapeComp
+			|| DestShapeComp == ShapeComp
+			|| ShapeComp->GetComponentLevel() != DestShapeComp->GetComponentLevel())
+		{
+			continue;
+		}
+
+		const FTransform& DestTransform = DestShapeComp->GetComponentTransform();
+		TConstArrayView<FZoneShapeConnector> DestConnectors = DestShapeComp->GetShapeConnectors();
+		TConstArrayView<FZoneShapeConnection> DestConnections = DestShapeComp->GetConnectedShapes();
+
+		for (int32 ConIndex = 0; ConIndex < DestConnectors.Num(); ConIndex++)
+		{
+			const FZoneShapeConnector& DestConnector = DestConnectors[ConIndex];
+			if (SourceConnector == &DestConnector
+				|| SourceConnector->LaneProfile != DestConnector.LaneProfile)
+			{
+				continue;;
+			}
+
+			const bool bOccupied = ConIndex < DestConnections.Num() && DestConnections[ConIndex].ShapeComponent.IsValid();
+			if (bOccupied)
+			{
+				continue;
+			}
+
+			// Check that the profile orientation matches before connecting.
+			const FZoneLaneProfile* LaneProfile = ZoneGraphSettings->GetLaneProfileByRef(SourceConnector->LaneProfile);
+			if (LaneProfile
+				&& (LaneProfile->IsSymmetrical() || SourceConnector->bReverseLaneProfile != DestConnector.bReverseLaneProfile))
+			{
+				const FVector DestWorldPosition = DestTransform.TransformPosition(DestConnector.Position);
+				const double Distance = FVector::Dist(SourceWorldPosition, DestWorldPosition);
+
+				if (Distance < BuildSettings.DragEndpointAutoConnectRange)
+				{
+					const FVector DestWorldNormal = DestTransform.TransformVector(DestConnector.Normal);
+					const FVector DestWorldUp = DestTransform.TransformVector(DestConnector.Up);
+					const int32 InfoIndex = AutoConnectState.DestShapeConnectorInfos.Add({ DestWorldPosition, DestWorldNormal, DestWorldUp });
+
+					if (Distance < ShortestDistance)
+					{
+						ShortestDistance = Distance;
+						AutoConnectState.ClosestShapeConnectorInfoIndex = InfoIndex;
+						AutoConnectState.NearestPointWorldPosition = DestWorldPosition;
+						AutoConnectState.NearestPointWorldNormal = DestWorldNormal;
+					}
+				}
+
+			}
+		}
+	}
+}
+
+void FZoneShapeComponentVisualizer::DetectCloseByShapeForAutoIntersectionCreation(const UZoneShapeComponent* ShapeComp, const FZoneShapePoint& DraggedPoint)
+{
+	ClearAutoIntersectionStatus();
+	bIsCreatingIntersection = true;
+
+	const UZoneGraphSettings* ZoneGraphSettings = GetDefault<UZoneGraphSettings>();
+	if (!ZoneGraphSettings)
+	{
+		return;
+	}
+	UZoneGraphSubsystem* ZoneGraph = UWorld::GetSubsystem<UZoneGraphSubsystem>(ShapeComp->GetWorld());
+	if (!ZoneGraph)
+	{
+		return;
+	}
+
+	const FZoneGraphBuildSettings& BuildSettings = ZoneGraphSettings->GetBuildSettings();
+	const TArray<FZoneGraphBuilderRegisteredComponent>& RegisteredShapeComponents = ZoneGraph->GetBuilder().GetRegisteredZoneShapeComponents();
+
+	const FTransform& SourceTransform = ShapeComp->GetComponentTransform();
+	FVector DraggedPointWorldPosition = SourceTransform.TransformPosition(DraggedPoint.Position);
+
+	TArray<uint32> QueryResults;
+	const FBox Bounds = FBox::BuildAABB(DraggedPointWorldPosition, FVector(BuildSettings.DragEndpointAutoIntersectionRange));
+	ZoneGraph->GetBuilder().QueryHashGrid(Bounds, QueryResults);
+	
+	double ClosestDistanceToSegment = std::numeric_limits<double>::infinity();
+	for (uint32 ComponentIndex : QueryResults)
+	{
+		check(RegisteredShapeComponents.IsValidIndex(int32(ComponentIndex)));
+		UZoneShapeComponent* DestShapeComp = RegisteredShapeComponents[ComponentIndex].Component;
+		if (!DestShapeComp ||
+			DestShapeComp == ShapeComp ||
+			ShapeComp->GetComponentLevel() != DestShapeComp->GetComponentLevel())
+		{
+			continue;
+		}
+
+		const FTransform& DestTransform = DestShapeComp->GetComponentTransform();
+		TConstArrayView<FZoneShapePoint> DestPoints = DestShapeComp->GetPoints();
+
+		FVector DraggedPointRelativePosition = DestTransform.InverseTransformPosition(DraggedPointWorldPosition);
+
+		if (DestShapeComp->GetShapeType() == FZoneShapeType::Spline)
+		{
+			// Spline
+			const FZoneLaneProfile* LaneProfile = ZoneGraphSettings->GetLaneProfileByRef(DestShapeComp->GetCommonLaneProfile());
+			const double HalfLanesTotalWidth = LaneProfile ? LaneProfile->GetLanesTotalWidth() * 0.5 : 0.0;
+
+			// Find closest point to the stem of the spline.
+			for (int32 Index = 0; Index < DestPoints.Num() - 1; Index++)
+			{
+				const FZoneShapePoint& CurrPoint = DestPoints[Index];
+				const FZoneShapePoint& NextPoint = DestPoints[Index + 1];
+				
+				FVector ClosestPoint;
+				float ClosestT = 0.0f;
+				UE::CubicBezier::ClosestPointApproximate(
+					DraggedPointRelativePosition,
+					CurrPoint.Position,
+					CurrPoint.GetOutControlPoint(),
+					NextPoint.Position,
+					NextPoint.GetInControlPoint(),
+					ClosestPoint,
+					ClosestT);
+
+				const double Dist = FVector::Dist(DraggedPointRelativePosition, ClosestPoint);
+				if (Dist < (BuildSettings.DragEndpointAutoIntersectionRange + HalfLanesTotalWidth)
+					&& Dist < ClosestDistanceToSegment)
+				{
+					ClosestDistanceToSegment = Dist;
+					CreateIntersectionState.WeakTargetShapeComponent = DestShapeComp;
+					CreateIntersectionState.OverlappingSegmentIndex = Index;
+					CreateIntersectionState.OverlappingSegmentT = ClosestT;
+					CreateIntersectionState.PreviewLocation = DestTransform.TransformPosition(ClosestPoint);
+				}
+			}
+		}
+		else
+		{
+			// Polygon
+			// Polygon defines the outline of the polygon, to make the behavior comparable to the spline case,
+			// just use linear segments between the lane profile points.
+			TArray<FZoneLaneProfile> PolyLaneProfiles;
+			DestShapeComp->GetPolygonLaneProfiles(PolyLaneProfiles);
+			check(DestPoints.Num() == PolyLaneProfiles.Num());
+
+			int32 PrevLaneProfilePointIndex = INDEX_NONE;
+			if (!DestPoints.IsEmpty() && DestPoints.Last().Type == FZoneShapePointType::LaneProfile)
+			{
+				PrevLaneProfilePointIndex = DestPoints.Num() - 1;
+			}
+			
+			for (int32 Index = 0; Index < DestPoints.Num(); Index++)
+			{
+				const FZoneShapePoint& CurrPoint = DestPoints[Index];
+				if (CurrPoint.Type == FZoneShapePointType::LaneProfile)
+				{
+					if (PrevLaneProfilePointIndex != INDEX_NONE)
+					{
+						const FZoneShapePoint& PrevPoint = DestPoints[PrevLaneProfilePointIndex];
+						const FVector ClosestPoint = FMath::ClosestPointOnSegment(DraggedPointRelativePosition, PrevPoint.Position, CurrPoint.Position);
+
+						const double PrevHalfLanesTotalWidth = PolyLaneProfiles[PrevLaneProfilePointIndex].GetLanesTotalWidth();
+						const double CurrHalfLanesTotalWidth = PolyLaneProfiles[Index].GetLanesTotalWidth();
+						const double HalfLanesTotalWidth = FMath::Min(PrevHalfLanesTotalWidth, CurrHalfLanesTotalWidth) * 0.5;
+						
+						const double Dist = FVector::Dist(DraggedPointRelativePosition, ClosestPoint);
+						if (Dist < (BuildSettings.DragEndpointAutoIntersectionRange + HalfLanesTotalWidth)
+							&& Dist < ClosestDistanceToSegment)
+						{
+							ClosestDistanceToSegment = Dist;
+							CreateIntersectionState.WeakTargetShapeComponent = DestShapeComp;
+							CreateIntersectionState.OverlappingSegmentIndex = -1; // Not used for polygons
+							CreateIntersectionState.OverlappingSegmentT = 0.0; // Not used for polygons
+							CreateIntersectionState.PreviewLocation = DestTransform.TransformPosition(ClosestPoint);
+						}
+					}
+					PrevLaneProfilePointIndex = Index;
+				}
+			}
+		}
+	}
+
+	// If the dragged point is close to a point on spline, or un-connected lane point in polygon, try to snap to that. 
+	if (UZoneShapeComponent* TargetShapeComponent = CreateIntersectionState.WeakTargetShapeComponent.Get())
+	{
+		const FTransform& TargetShapeCompTransform = TargetShapeComponent->GetComponentTransform();
+		CreateIntersectionState.ClosePointIndex = INDEX_NONE;
+		
+		TArray<FZoneShapePoint>& TargetShapePoints = TargetShapeComponent->GetMutablePoints();
+		const int32 NumPoints = TargetShapePoints.Num();
+
+		TConstArrayView<FZoneShapeConnector> DestConnectors = TargetShapeComponent->GetShapeConnectors();
+		TConstArrayView<FZoneShapeConnection> DestConnections = TargetShapeComponent->GetConnectedShapes();
+		
+		static const double SnapToleranceSqr = FMath::Square(BuildSettings.SnapAutoIntersectionToClosestPointTolerance);
+		double ShortestDistanceSqr = SnapToleranceSqr;
+		
+		for (int32 PointIndex = 0; PointIndex < NumPoints; PointIndex++)
+		{
+			const FZoneShapePoint& CurrTargetPoint = TargetShapePoints[PointIndex];
+			
+			// Only allow to snap to lane profile points on polygons.
+			if (TargetShapeComponent->GetShapeType() == FZoneShapeType::Polygon
+				&& CurrTargetPoint.Type != FZoneShapePointType::LaneProfile)
+			{
+				continue;
+			}
+
+			// Prevent snapping to already connected points.
+			bool bOccupied = false;
+			for (int ConIndex = 0; ConIndex < DestConnectors.Num(); ConIndex++)
+			{
+				if (DestConnectors[ConIndex].PointIndex == PointIndex)
+				{
+					bOccupied = ConIndex < DestConnections.Num() && DestConnections[ConIndex].ShapeComponent.IsValid();
+					if (bOccupied)
+					{
+						break;
+					}
+				}
+			}
+			if (bOccupied)
+			{
+				continue;
+			}
+
+			const FVector TargetPointWorldPosition = TargetShapeCompTransform.TransformPosition(CurrTargetPoint.Position);
+			const double DistSqr = FVector::DistSquared(DraggedPointWorldPosition, TargetPointWorldPosition);
+			
+			if (DistSqr < SnapToleranceSqr
+				&& DistSqr < ShortestDistanceSqr)
+			{
+				ShortestDistanceSqr = DistSqr;
+				CreateIntersectionState.ClosePointIndex = PointIndex;
+				CreateIntersectionState.PreviewLocation = TargetPointWorldPosition;
+			}
+		}
+	}
+}
+
 bool FZoneShapeComponentVisualizer::TransformSelectedControlPoint(const FVector& DeltaTranslate)
 {
 	if (UZoneShapeComponent* ShapeComp = GetEditedShapeComponent())
@@ -1087,7 +1458,14 @@ bool FZoneShapeComponentVisualizer::HandleInputKey(FEditorViewportClient* Viewpo
 	{
 		return false;
 	}
-	
+
+	const UZoneGraphSettings* ZoneGraphSettings = GetDefault<UZoneGraphSettings>();
+	if (!ZoneGraphSettings)
+	{
+		return false;
+	}
+	const FZoneGraphBuildSettings& BuildSettings = ZoneGraphSettings->GetBuildSettings();
+
 	if (IsAnySelectedPointIndexOutOfRange(*ShapeComp))
 	{
 		// Something external has changed the number of shape points, meaning that the cached selected keys are no longer valid
@@ -1107,65 +1485,54 @@ bool FZoneShapeComponentVisualizer::HandleInputKey(FEditorViewportClient* Viewpo
 		bHasCachedRotation = false;
 		CachedRotation = FQuat::Identity;
 
-		if (bIsAutoConnecting && SelectedPointForConnecting >= 0 && SelectedPointForConnecting < ShapeComp->GetNumPoints())
+		TArray<FZoneShapePoint>& ShapePoints = ShapeComp->GetMutablePoints();
+		if (ShapePoints.IsValidIndex(SelectedPointForConnecting))
 		{
-			const UZoneGraphSettings* ZoneGraphSettings = GetDefault<UZoneGraphSettings>();
-			check(ZoneGraphSettings);
-
-			const FZoneLaneProfile* LaneProfile = ZoneGraphSettings->GetLaneProfileByRef(ShapeComp->GetCommonLaneProfile());
-			check(LaneProfile);
-			float HalfLanesTotalWidth = LaneProfile->GetLanesTotalWidth() * 0.5;
-
-			FZoneShapePoint& DraggedPoint = ShapeComp->GetMutablePoints()[SelectedPointForConnecting];
-
-#if WITH_EDITOR
-			if (const FZoneShapeConnector* SourceConnector = ShapeComp->GetShapeConnectorByPointIndex(SelectedPointForConnecting))
+			if (bIsAutoConnecting)
 			{
-				const FTransform& SourceTransform = ShapeComp->GetComponentTransform();
-				const FVector SourceWorldNormal = SourceTransform.TransformVector(SourceConnector->Normal);
-
-				const FZoneGraphBuildSettings& BuildSettings = ZoneGraphSettings->GetBuildSettings();
-				static const float ConnectionSnapAngleCos = FMath::Cos(FMath::DegreesToRadians(BuildSettings.ConnectionSnapAngle));
+				const FZoneLaneProfile* LaneProfile = ZoneGraphSettings->GetLaneProfileByRef(ShapeComp->GetCommonLaneProfile());
+				double HalfLanesTotalWidth = LaneProfile ? LaneProfile->GetLanesTotalWidth() * 0.5 : 0.0;
 
-				if (ClosestShapeConnectorInfoIndex != INDEX_NONE)
+				FZoneShapePoint& DraggedPoint = ShapePoints[SelectedPointForConnecting];
+				const FZoneShapeConnector* SourceConnector = ShapeComp->GetShapeConnectorByPointIndex(SelectedPointForConnecting);
+				
+				if (SourceConnector && AutoConnectState.ClosestShapeConnectorInfoIndex != INDEX_NONE)
 				{
-					// Snap point location
-					DraggedPoint.Position = SourceTransform.InverseTransformPosition(NearestPointWorldPosition);
-					FVector Normal = SourceTransform.InverseTransformVector(NearestPointWorldNormal);
-					const FRotator Rotation = FRotationMatrix::MakeFromX(SelectedPointForConnecting == 0 ? Normal : -Normal).Rotator();
-					DraggedPoint.Rotation = Rotation;
-
-					// If the zone shape is a spline and the point type is not Bezier, setting the point rotation doesn't work.
-					// An extra point is needed to align the connectors and make it connect.
-					if (ShapeComp->GetShapeType() == FZoneShapeType::Spline &&
-						DraggedPoint.Type != FZoneShapePointType::Bezier &&
-						FVector::DotProduct(SourceWorldNormal, -NearestPointWorldNormal) <= ConnectionSnapAngleCos)
-					{
-						// Add extra point
-						TArray<FZoneShapePoint>& Points = ShapeComp->GetMutablePoints();
-						FZoneShapePoint ExtraPoint = DraggedPoint;
-						ExtraPoint.Position += Normal * HalfLanesTotalWidth;
-						ExtraPoint.Rotation = Rotation;
-						Points.Insert(ExtraPoint, ShapeComp->GetNumPoints() - 1);
-					}
-
-					// Update shape
-					ShapeComp->UpdateShape();
+					const FTransform& SourceTransform = ShapeComp->GetComponentTransform();
+					const FVector SourceWorldNormal = SourceTransform.TransformVector(SourceConnector->Normal);
+
+					const double ConnectionSnapAngleCos = FMath::Cos(FMath::DegreesToRadians(BuildSettings.ConnectionSnapAngle));
+
+					UE::ZoneGraph::Editor::Private::SnapConnect(
+						ShapeComp,
+						DraggedPoint,
+						SourceTransform,
+						SourceWorldNormal,
+						AutoConnectState.NearestPointWorldPosition,
+						AutoConnectState.NearestPointWorldNormal,
+						ConnectionSnapAngleCos,
+						HalfLanesTotalWidth);
 				}
 			}
-#endif
-		}
 
-		bIsAutoConnecting = false;
-		DestShapeConnectorInfos.Empty();
-		ClosestShapeConnectorInfoIndex = INDEX_NONE;
+			if (bIsCreatingIntersection)
+			{
+				CreateIntersection(ShapeComp);
+			}
+		}
+		
+		ClearAutoConnectingStatus();
+		ClearAutoIntersectionStatus();
 	}
 
 	if (Key == EKeys::C && Event == IE_Released)
 	{
-		bIsAutoConnecting = false;
-		DestShapeConnectorInfos.Empty();
-		ClosestShapeConnectorInfoIndex = INDEX_NONE;
+		ClearAutoConnectingStatus();
+	}
+
+	if (Key == EKeys::X && Event == IE_Released)
+	{
+		ClearAutoIntersectionStatus();
 	}
 
 	if (Key == EKeys::LeftMouseButton && Event == IE_Pressed)
@@ -1207,6 +1574,8 @@ bool FZoneShapeComponentVisualizer::HandleInputKey(FEditorViewportClient* Viewpo
 			if (World->LineTraceSingleByChannel(Hit, Start, End, ECollisionChannel::ECC_WorldStatic, QueryParams))
 			{
 				// Add a new point at the position
+				const FScopedTransaction Transaction(LOCTEXT("AddShapePoint", "Add Shape Point And Snap To Floor"));
+				ShapeComp->Modify();
 				TArray<FZoneShapePoint>& Points = ShapeComp->GetMutablePoints();
 				FZoneShapePoint PointToAdd(ShapeComp->GetComponentTransform().InverseTransformPosition(Hit.Location));
 				Points.Add(PointToAdd);
@@ -1585,9 +1954,12 @@ bool FZoneShapeComponentVisualizer::DuplicatePointForAltDrag(const FVector& InDr
 	return true;
 }
 
-void FZoneShapeComponentVisualizer::SplitSegment(const int32 InSegmentIndex, const float SegmentSplitT) const
+void FZoneShapeComponentVisualizer::SplitSegment(const int32 InSegmentIndex, const float SegmentSplitT, UZoneShapeComponent* ShapeComp) const
 {
-	UZoneShapeComponent* ShapeComp = GetEditedShapeComponent();
+	if (!ShapeComp)
+	{
+		ShapeComp = GetEditedShapeComponent();
+	}
 
 	check(ShapeComp != nullptr);
 	check(InSegmentIndex != INDEX_NONE);
@@ -1631,7 +2003,7 @@ void FZoneShapeComponentVisualizer::SplitSegment(const int32 InSegmentIndex, con
 		NewPoint.TangentLength = 0.0f;
 	}
 
-	const int NewPointIndex = InSegmentIndex + 1;
+	const int32 NewPointIndex = InSegmentIndex + 1;
 
 	ShapePoints.Insert(NewPoint, NewPointIndex);
 
@@ -1856,9 +2228,12 @@ void FZoneShapeComponentVisualizer::OnBreakAtPointNewComponents() const
 	BreakAtPoint(false);
 }
 
-void FZoneShapeComponentVisualizer::BreakAtPoint(bool bCreateNewActor) const
+TArray<UZoneShapeComponent*>  FZoneShapeComponentVisualizer::BreakAtPoint(bool bCreateNewActor, UZoneShapeComponent* ShapeComp) const
 {
-	UZoneShapeComponent* ShapeComp = GetEditedShapeComponent();
+	if (!ShapeComp)
+	{
+		ShapeComp = GetEditedShapeComponent();
+	}
 	check(ShapeComp != nullptr);
 	check(SelectionState);
 	const TSet<int32>& SelectedPoints = SelectionState->GetSelectedPoints();
@@ -1869,6 +2244,9 @@ void FZoneShapeComponentVisualizer::BreakAtPoint(bool bCreateNewActor) const
 	check(SelectedPoints.Num() > 0);
 	check(SelectedPoints.Contains(LastPointIndexSelected));
 
+	TArray<UZoneShapeComponent*> ShapeComponents;
+	ShapeComponents.Add(ShapeComp);
+
 	ShapeComp->Modify();
 	if (AActor* Owner = ShapeComp->GetOwner())
 	{
@@ -1928,6 +2306,8 @@ void FZoneShapeComponentVisualizer::BreakAtPoint(bool bCreateNewActor) const
 			NewShapeComponent->AttachToComponent(ShapeComp, FAttachmentTransformRules::KeepWorldTransform);
 			NewShapeComponent->Modify();
 		}
+		NewShapeComponent->SetCommonLaneProfile(ShapeComp->GetCommonLaneProfile());
+		ShapeComponents.Add(NewShapeComponent);
 
 		// Copy points
 		TArray<FZoneShapePoint>& NewShapePoints = NewShapeComponent->GetMutablePoints();
@@ -1939,13 +2319,6 @@ void FZoneShapeComponentVisualizer::BreakAtPoint(bool bCreateNewActor) const
 		}
 		NewShapeComponent->UpdateShape();
 
-		if (i == 0 || (i == 1 && SelectedPointsSorted[0] == 0))
-		{
-			// Keep the last segment on the original shape component
-			ShapePoints.RemoveAt(EndIndex);
-			break;
-		}
-
 		// Delete all points after the selected one
 		for (int32 Index = EndIndex; Index > SelectedIndex; Index--)
 		{
@@ -1972,6 +2345,8 @@ void FZoneShapeComponentVisualizer::BreakAtPoint(bool bCreateNewActor) const
 	FLevelEditorModule& LevelEditor = FModuleManager::LoadModuleChecked<FLevelEditorModule>("LevelEditor");
 	LevelEditor.BroadcastComponentsEdited();
 	LevelEditor.BroadcastRedrawViewports(false);
+
+	return ShapeComponents;
 }
 
 bool FZoneShapeComponentVisualizer::CanBreakAtPoint() const
@@ -1980,7 +2355,7 @@ bool FZoneShapeComponentVisualizer::CanBreakAtPoint() const
 	const TSet<int32>& SelectedPoints = SelectionState->GetSelectedPoints();
 	const int32 LastPointIndexSelected = SelectionState->GetLastPointIndexSelected();
 	UZoneShapeComponent* ShapeComp = GetEditedShapeComponent();
-	return (ShapeComp != nullptr && SelectedPoints.Num() > 0 && LastPointIndexSelected != INDEX_NONE);
+	return (ShapeComp != nullptr && ShapeComp->GetShapeType() == FZoneShapeType::Spline && SelectedPoints.Num() > 0 && LastPointIndexSelected != INDEX_NONE);
 }
 
 void FZoneShapeComponentVisualizer::OnBreakAtSegmentNewActors() const
@@ -2007,7 +2382,7 @@ void FZoneShapeComponentVisualizer::BreakAtSegment(bool bCreateNewActor) const
 	SelectionState->Modify();
 	int32 SegmentIndex = SelectionState->GetSelectedSegmentIndex();
 	SplitSegment(SegmentIndex, SelectionState->GetSelectedSegmentT());
-	const int NewPointIndex = SegmentIndex + 1;
+	const int32 NewPointIndex = SegmentIndex + 1;
 	ChangeSelectionState(NewPointIndex, false);
 	BreakAtPoint(bCreateNewActor);
 	SelectionState->SetSelectedSegmentPoint(FVector::ZeroVector);
@@ -2016,7 +2391,14 @@ void FZoneShapeComponentVisualizer::BreakAtSegment(bool bCreateNewActor) const
 
 bool FZoneShapeComponentVisualizer::CanBreakAtSegment() const
 {
-	return CanAddPointToSegment();
+	const UZoneShapeComponent* ShapeComp = GetEditedShapeComponent();
+	if (ShapeComp != nullptr && ShapeComp->GetShapeType() == FZoneShapeType::Spline)
+	{
+		check(SelectionState);
+		const int32 SelectedSegmentIndex = SelectionState->GetSelectedSegmentIndex();
+		return (SelectedSegmentIndex != INDEX_NONE && SelectedSegmentIndex >= 0 && SelectedSegmentIndex < ShapeComp->GetNumPoints());
+	}
+	return false;
 }
 
 TSharedPtr<SWidget> FZoneShapeComponentVisualizer::GenerateContextMenu() const
@@ -2031,10 +2413,13 @@ TSharedPtr<SWidget> FZoneShapeComponentVisualizer::GenerateContextMenu() const
 		{
 			MenuBuilder.AddMenuEntry(FZoneShapeComponentVisualizerCommands::Get().AddPoint);
 
-			MenuBuilder.AddSubMenu(
-				LOCTEXT("BreakAtPoint", "Break At Point"),
-				LOCTEXT("BreakAtPointTooltip", "Break the shape into pieces at the currently selected points."),
-				FNewMenuDelegate::CreateSP(this, &FZoneShapeComponentVisualizer::GenerateBreakAtSegmentSubMenu));
+			if (CanBreakAtSegment())
+			{
+				MenuBuilder.AddSubMenu(
+					LOCTEXT("BreakAtPoint", "Break At Point"),
+					LOCTEXT("BreakAtPointTooltip", "Break the shape into pieces at the currently selected points."),
+					FNewMenuDelegate::CreateSP(this, &FZoneShapeComponentVisualizer::GenerateBreakAtSegmentSubMenu));
+			}
 		}
 		else if (SelectionState->GetLastPointIndexSelected() != INDEX_NONE)
 		{
@@ -2052,10 +2437,13 @@ TSharedPtr<SWidget> FZoneShapeComponentVisualizer::GenerateContextMenu() const
 				LOCTEXT("SplineSnapAlignTooltip", "Snap align options."),
 				FNewMenuDelegate::CreateSP(this, &FZoneShapeComponentVisualizer::GenerateSnapAlignSubMenu));
 
-			MenuBuilder.AddSubMenu(
-				LOCTEXT("BreakAtPoint", "Break At Point"),
-				LOCTEXT("BreakAtPointTooltip", "Break the shape into pieces at the currently selected points."),
-				FNewMenuDelegate::CreateSP(this, &FZoneShapeComponentVisualizer::GenerateBreakAtPointSubMenu));
+			if (CanBreakAtPoint())
+			{
+				MenuBuilder.AddSubMenu(
+					LOCTEXT("BreakAtPoint", "Break At Point"),
+					LOCTEXT("BreakAtPointTooltip", "Break the shape into pieces at the currently selected points."),
+					FNewMenuDelegate::CreateSP(this, &FZoneShapeComponentVisualizer::GenerateBreakAtPointSubMenu));
+			}
 		}
 	}
 	MenuBuilder.EndSection();
@@ -2091,12 +2479,8 @@ void FZoneShapeComponentVisualizer::GenerateSnapAlignSubMenu(FMenuBuilder& MenuB
 
 void FZoneShapeComponentVisualizer::GenerateBreakAtPointSubMenu(FMenuBuilder& MenuBuilder) const
 {
-	UZoneShapeComponent* ShapeComp = GetEditedShapeComponent();
-	if (ShapeComp && ShapeComp->GetShapeType() == FZoneShapeType::Spline)
-	{
-		MenuBuilder.AddMenuEntry(FZoneShapeComponentVisualizerCommands::Get().BreakAtPointNewActors);
-		MenuBuilder.AddMenuEntry(FZoneShapeComponentVisualizerCommands::Get().BreakAtPointNewComponents);
-	}
+	MenuBuilder.AddMenuEntry(FZoneShapeComponentVisualizerCommands::Get().BreakAtPointNewActors);
+	MenuBuilder.AddMenuEntry(FZoneShapeComponentVisualizerCommands::Get().BreakAtPointNewComponents);
 }
 
 void FZoneShapeComponentVisualizer::GenerateBreakAtSegmentSubMenu(FMenuBuilder& MenuBuilder) const
@@ -2109,4 +2493,304 @@ void FZoneShapeComponentVisualizer::GenerateBreakAtSegmentSubMenu(FMenuBuilder&
 	}
 }
 
+void FZoneShapeComponentVisualizer::CreateIntersection(UZoneShapeComponent* ShapeComp)
+{
+	if (UZoneShapeComponent* TargetShapeComponent = CreateIntersectionState.WeakTargetShapeComponent.Get())
+	{
+		const FScopedTransaction Transaction(LOCTEXT("CreateIntersection", "Create an Intersection With The Dragged Point and Overlapped Shape"));
+		TargetShapeComponent->Modify();
+		FZoneShapePoint& DraggedPoint = ShapeComp->GetMutablePoints()[SelectedPointForConnecting];
+		if (TargetShapeComponent->GetShapeType() == FZoneShapeType::Spline)
+		{
+			CreateIntersectionForSplineShape(ShapeComp, DraggedPoint);
+		}
+		else
+		{
+			CreateIntersectionForPolygonShape(ShapeComp, DraggedPoint);
+		}
+	}
+}
+
+void FZoneShapeComponentVisualizer::CreateIntersectionForSplineShape(UZoneShapeComponent* ShapeComp, FZoneShapePoint& DraggedPoint, bool DestroyCoveredShape)
+{
+	const UZoneGraphSettings* ZoneGraphSettings = GetDefault<UZoneGraphSettings>();
+	if (!ZoneGraphSettings)
+	{
+		return;
+	}
+
+	UZoneShapeComponent* TargetShapeComponent = CreateIntersectionState.WeakTargetShapeComponent.Get();
+	if (!TargetShapeComponent)
+	{
+		return;
+	}
+	
+	if (CreateIntersectionState.OverlappingSegmentIndex == INDEX_NONE)
+	{
+		return;
+	}
+
+	const FZoneLaneProfile* LaneProfile = ZoneGraphSettings->GetLaneProfileByRef(ShapeComp->GetCommonLaneProfile());
+	double const HalfLanesTotalWidth = LaneProfile ? LaneProfile->GetLanesTotalWidth() * 0.5 : 0.0;
+
+	// Get overlapping position on the target segment
+	FVector NewPointPosition = UE::ZoneGraph::Editor::Private::GetPositionOnSegment(TargetShapeComponent->GetMutablePoints(), CreateIntersectionState.OverlappingSegmentIndex, CreateIntersectionState.OverlappingSegmentT);
+
+	bool bCloseToPoint = CreateIntersectionState.ClosePointIndex != INDEX_NONE;
+	if (bCloseToPoint)
+	{
+		// If close to a point, select it as the point to break at.
+		ChangeSelectionState(CreateIntersectionState.ClosePointIndex, false);
+	}
+	else
+	{
+		// At the overlapping position, add a point to break at.
+		SplitSegment(CreateIntersectionState.OverlappingSegmentIndex, CreateIntersectionState.OverlappingSegmentT, TargetShapeComponent);
+	}
+
+	// Break the zone shape
+	const FTransform& ShapeCompTransform = ShapeComp->GetComponentTransform();
+	TArray<UZoneShapeComponent*> ShapeComponents = BreakAtPoint(true, TargetShapeComponent);
+
+	// Create an intersection
+	FActorSpawnParameters SpawnParams;
+	AZoneShape* IntersectionShapeActor = ShapeComp->GetWorld()->SpawnActor<AZoneShape>(AZoneShape::StaticClass(), ShapeCompTransform, SpawnParams);
+	UZoneShapeComponent* IntersectionShapeComponent = IntersectionShapeActor->GetComponentByClass<UZoneShapeComponent>();
+	IntersectionShapeComponent->SetShapeType(FZoneShapeType::Polygon);
+	const FTransform& IntersectionTransform = IntersectionShapeComponent->GetComponentTransform();
+
+	FVector Normal = ShapeComp->GetShapeConnectorByPointIndex(SelectedPointForConnecting)->Normal;
+	if (ShapeComponents.Num() == 1 && bCloseToPoint)
+	{
+		// The point was dragged to the start or end of a zone shape. Create an intersection that connects these two shapes.
+
+		// Get the target zone shape's connector that is close to the dragged point.
+		int32 PointIndex = CreateIntersectionState.OverlappingSegmentT < 0.5f ? CreateIntersectionState.OverlappingSegmentIndex : CreateIntersectionState.OverlappingSegmentIndex + 1;
+		const FZoneShapeConnector* TargetConnector = TargetShapeComponent->GetShapeConnectorByPointIndex(PointIndex);
+
+		// Compute the intersection location from the connector position and normal
+		FVector TargetNormal = TargetConnector->Normal;
+		FVector TargetWorldNormal = TargetShapeComponent->GetComponentTransform().TransformVector(TargetConnector->Normal);
+		IntersectionShapeActor->SetActorLocation(NewPointPosition + TargetWorldNormal * HalfLanesTotalWidth);
+
+
+		// Connect
+		TArray<FZoneShapePoint>& Points0 = ShapeComponents[0]->GetMutablePoints();
+		int32 Index0 = Points0.Num() - 1;
+		FVector Normal0 = ShapeComponents[0]->GetShapeConnectorByPointIndex(Index0)->Normal;
+		Points0.Last().Position -= Normal0 * HalfLanesTotalWidth;
+		ShapeComponents[0]->UpdateShape();
+
+		TArray<FZoneShapePoint>& Points = IntersectionShapeComponent->GetMutablePoints();
+		UE::ZoneGraph::Editor::Private::SetPolygonPointLaneProfileToMatchSpline(Points[0], IntersectionShapeComponent, ShapeComponents[0]);
+		TConstArrayView<FZoneShapePoint> TargetPoints = TargetShapeComponent->GetPoints();
+		const FTransform Shape0Transform = ShapeComponents[0]->GetComponentTransform();
+		const FVector Point0WorldPosition = Shape0Transform.TransformPosition(TargetPoints[PointIndex].Position);
+		const FVector Point0WorldNormal = Shape0Transform.TransformVector(TargetNormal);
+		UE::ZoneGraph::Editor::Private::SetPointPositionRotation(Points[0], IntersectionTransform, Point0WorldPosition, Point0WorldNormal);
+
+		UE::ZoneGraph::Editor::Private::SetPolygonPointLaneProfileToMatchSpline(Points[1], IntersectionShapeComponent, ShapeComp);
+		DraggedPoint.Position -= Normal * HalfLanesTotalWidth;
+		const FVector Point1WorldPosition = ShapeCompTransform.TransformPosition(DraggedPoint.Position);
+		const FVector Point1WorldNormal = ShapeCompTransform.TransformVector(Normal);
+		UE::ZoneGraph::Editor::Private::SetPointPositionRotation(Points[1], IntersectionTransform, Point1WorldPosition, Point1WorldNormal);
+
+		// Update shape
+		ShapeComp->UpdateShape();
+
+		// Update point positions
+		IntersectionShapeComponent->UpdateShape();
+	}
+	else if (ShapeComponents.Num() == 2)
+	{
+		// Cut the intersected shape
+		FVector DraggedPointWorldPosition = ShapeCompTransform.TransformPosition(DraggedPoint.Position);
+		const FTransform& TargetTransform = TargetShapeComponent->GetComponentTransform();
+		FBox Bounds = FBox::BuildAABB(TargetTransform.TransformPosition(NewPointPosition), FVector(HalfLanesTotalWidth));
+		// Move points
+		const FTransform& ShapeTransform0 = ShapeComponents[0]->GetComponentTransform();
+		const FTransform& ShapeTransform1 = ShapeComponents[1]->GetComponentTransform();
+		TArray<FZoneShapePoint>& Points0 = ShapeComponents[0]->GetMutablePoints();
+		int32 Index0 = Points0.Num() - 1;
+		for (int32 i = Index0 - 1; i > 1; i--)
+		{
+			if (!Bounds.IsInside(ShapeTransform0.TransformPosition(Points0[i].Position)))
+			{
+				continue;
+			}
+			Points0.RemoveAt(i);
+		}
+		Index0 = Points0.Num() - 1;
+		if (ShapeComponents[0])
+		{
+			FVector Normal0 = ShapeComponents[0]->GetShapeConnectorByPointIndex(Index0)->Normal;
+			FVector Offset = Normal0 * HalfLanesTotalWidth;
+			if (Points0.Num() == 2)
+			{
+				double Length = FVector::Dist(Points0[0].Position, Points0[1].Position);
+				if (Length < HalfLanesTotalWidth * 2)
+				{
+					Offset = Normal0 * Length * 0.5;
+				}
+			}
+			Points0.Last().Position -= Offset;
+			ShapeComponents[0]->UpdateShape();
+		}
+
+		TArray<FZoneShapePoint>& Points1 = ShapeComponents[1]->GetMutablePoints();
+		int32 Index1 = 0;
+		for (int32 i = Index1 + 1; i < (Points1.Num() - 2); i++)
+		{
+			if (!Bounds.IsInside(ShapeTransform1.TransformPosition(Points1[i].Position)))
+			{
+				continue;
+			}
+			Points1.RemoveAt(i);
+		}
+		if (ShapeComponents[1])
+		{
+			FVector Normal1 = ShapeComponents[1]->GetShapeConnectorByPointIndex(Index1)->Normal;
+			FVector Offset = Normal1 * HalfLanesTotalWidth;
+			if (Points1.Num() == 2)
+			{
+				double Length = FVector::Dist(Points1[0].Position, Points1[1].Position);
+				if (Length < HalfLanesTotalWidth * 2)
+				{
+					Offset = Normal1 * Length * 0.5;
+				}
+			}
+			Points1[Index1].Position -= Offset;
+			ShapeComponents[1]->UpdateShape();
+		}
+
+		// Create intersection with the same profile
+		IntersectionShapeActor->SetActorLocation(DraggedPointWorldPosition);
+
+		// Get points. Set positions. Set profile.
+		TArray<FZoneShapePoint>& Points = IntersectionShapeComponent->GetMutablePoints();
+		if (ShapeComponents[0] && ShapeComponents[1])
+		{
+			Points.Add(FZoneShapePoint(Points[1]));
+		}
+
+		// Connect
+		int32 IntersectionPointIndex = 0;
+		if (ShapeComponents[0])
+		{
+			UE::ZoneGraph::Editor::Private::SetPolygonPointLaneProfileToMatchSpline(Points[IntersectionPointIndex], IntersectionShapeComponent, ShapeComponents[0]);
+
+			const FVector PointWorldPosition = ShapeTransform0.TransformPosition(Points0.Last().Position);
+			const FZoneShapeConnector* Connector0 = ShapeComponents[0]->GetShapeConnectorByPointIndex(Points0.Num() - 1);
+			const FVector PointWorldNormal = ShapeTransform0.TransformVector(Connector0->Normal);
+			UE::ZoneGraph::Editor::Private::SetPointPositionRotation(Points[IntersectionPointIndex], IntersectionTransform, PointWorldPosition, PointWorldNormal);
+			IntersectionPointIndex++;
+		}
+
+		if (ShapeComponents[1])
+		{
+			UE::ZoneGraph::Editor::Private::SetPolygonPointLaneProfileToMatchSpline(Points[IntersectionPointIndex], IntersectionShapeComponent, ShapeComponents[1]);
+
+			const FVector PointWorldPosition = ShapeTransform1.TransformPosition(Points1[0].Position);
+			const FZoneShapeConnector* Connector1 = ShapeComponents[1]->GetShapeConnectorByPointIndex(0);
+			const FVector PointWorldNormal = ShapeTransform1.TransformVector(Connector1->Normal);
+			UE::ZoneGraph::Editor::Private::SetPointPositionRotation(Points[IntersectionPointIndex], IntersectionTransform, PointWorldPosition, PointWorldNormal);
+			IntersectionPointIndex++;
+		}
+
+		UE::ZoneGraph::Editor::Private::SetPolygonPointLaneProfileToMatchSpline(Points[IntersectionPointIndex], IntersectionShapeComponent, ShapeComp);
+		DraggedPoint.Position -= Normal * HalfLanesTotalWidth;
+		ShapeComp->UpdateShape(); // Update shape
+		DraggedPointWorldPosition = ShapeCompTransform.TransformPosition(DraggedPoint.Position);
+		const FVector WorldNormal = ShapeCompTransform.TransformVector(Normal);
+		UE::ZoneGraph::Editor::Private::SetPointPositionRotation(Points[IntersectionPointIndex], IntersectionTransform, DraggedPointWorldPosition, WorldNormal);
+
+		UE::ZoneGraph::Editor::Private::SortPolygonPointsCounterclockwise(IntersectionShapeComponent);
+
+		// Update point positions
+		IntersectionShapeComponent->UpdateShape();
+	}
+}
+
+void FZoneShapeComponentVisualizer::CreateIntersectionForPolygonShape(UZoneShapeComponent* ShapeComp, FZoneShapePoint& DraggedPoint)
+{
+	UZoneShapeComponent* TargetShapeComponent = CreateIntersectionState.WeakTargetShapeComponent.Get();
+	if (!TargetShapeComponent)
+	{
+		return;
+	}
+
+	const UZoneGraphSettings* ZoneGraphSettings = GetDefault<UZoneGraphSettings>();
+	if (!ZoneGraphSettings)
+	{
+		return;
+	}
+
+	const FTransform& TargetShapeCompTransform = TargetShapeComponent->GetComponentTransform();
+	const FZoneGraphBuildSettings& BuildSettings = ZoneGraphSettings->GetBuildSettings();
+	TArray<FZoneShapePoint>& TargetShapePoints = TargetShapeComponent->GetMutablePoints();
+
+	const FTransform& SourceTransform = ShapeComp->GetComponentTransform();
+	const FZoneShapeConnector* SourceConnector = ShapeComp->GetShapeConnectorByPointIndex(SelectedPointForConnecting);
+	if (CreateIntersectionState.ClosePointIndex != INDEX_NONE)
+	{
+		// If the dragged point is close to a connector, connect.
+		const FVector TargetPointWorldPosition = TargetShapeCompTransform.TransformPosition(TargetShapePoints[CreateIntersectionState.ClosePointIndex].Position);
+		const FZoneShapeConnector* TargetConnector = TargetShapeComponent->GetShapeConnectorByPointIndex(CreateIntersectionState.ClosePointIndex);
+		const FVector TargetPointWorldNormal = TargetShapeComponent->GetComponentTransform().TransformVector(TargetConnector->Normal);
+
+		const double ConnectionSnapAngleCos = FMath::Cos(FMath::DegreesToRadians(BuildSettings.ConnectionSnapAngle));
+		const FZoneLaneProfile* LaneProfile = ZoneGraphSettings->GetLaneProfileByRef(ShapeComp->GetCommonLaneProfile());
+		const double HalfLanesTotalWidth = LaneProfile ? LaneProfile->GetLanesTotalWidth() * 0.5 : 0.0;
+		UE::ZoneGraph::Editor::Private::SnapConnect(
+			ShapeComp,
+			DraggedPoint,
+			SourceTransform,
+			SourceTransform.TransformVector(SourceConnector->Normal),
+			TargetPointWorldPosition,
+			TargetPointWorldNormal,
+			ConnectionSnapAngleCos,
+			HalfLanesTotalWidth);
+	}
+	else
+	{
+		// If the dragged point is not close to any connector, add a point and connect.
+		FZoneShapePoint NewPoint = FZoneShapePoint(TargetShapePoints[0]);
+
+		UE::ZoneGraph::Editor::Private::SetPolygonPointLaneProfileToMatchSpline(NewPoint, TargetShapeComponent, ShapeComp);
+		TargetShapePoints.Add(NewPoint);
+		
+		UE::ZoneGraph::Editor::Private::SetPointPositionRotation(
+			TargetShapePoints.Last(0),
+			TargetShapeCompTransform,
+			SourceTransform.TransformPosition(DraggedPoint.Position),
+			SourceTransform.TransformVector(SourceConnector->Normal));
+		
+		UE::ZoneGraph::Editor::Private::SortPolygonPointsCounterclockwise(TargetShapeComponent);
+		
+		TargetShapeComponent->UpdateShape();
+	}
+}
+
+void FZoneShapeComponentVisualizer::ClearAutoConnectingStatus()
+{
+	bIsAutoConnecting = false;
+	AutoConnectState = {};
+}
+
+void FZoneShapeComponentVisualizer::ClearAutoIntersectionStatus()
+{
+	bIsCreatingIntersection = false;
+	CreateIntersectionState = {};
+}
+
+bool FZoneShapeComponentVisualizer::CanAutoConnect(const UZoneShapeComponent* ShapeComp) const
+{
+	return ShapeComp->GetShapeType() == FZoneShapeType::Spline;
+}
+
+bool FZoneShapeComponentVisualizer::CanAutoCreateIntersection(const UZoneShapeComponent* ShapeComp) const
+{
+	return ShapeComp->GetShapeType() == FZoneShapeType::Spline;
+}
+
 #undef LOCTEXT_NAMESPACE
diff --git a/Engine/Plugins/Runtime/ZoneGraph/Source/ZoneGraphEditor/Private/ZoneShapeComponentVisualizer.h b/Engine/Plugins/Runtime/ZoneGraph/Source/ZoneGraphEditor/Private/ZoneShapeComponentVisualizer.h
--- a/Engine/Plugins/Runtime/ZoneGraph/Source/ZoneGraphEditor/Private/ZoneShapeComponentVisualizer.h
+++ b/Engine/Plugins/Runtime/ZoneGraph/Source/ZoneGraphEditor/Private/ZoneShapeComponentVisualizer.h
@@ -195,9 +195,9 @@ protected:
 
 	struct ZoneShapeConnectorRenderInfo
 	{
-		ZoneShapeConnectorRenderInfo(FVector InPosition, FVector InNormal, FVector InUp)
+		ZoneShapeConnectorRenderInfo(FVector InPosition, FVector InFoward, FVector InUp)
 			: Position(InPosition)
-			, Normal(InNormal)
+			, Foward(InFoward)
 			, Up(InUp)
 		{
 
@@ -206,8 +206,8 @@ protected:
 		// Position of the connector.
 		FVector Position = FVector::ZeroVector;
 
-		// Normal direction of the connector.
-		FVector Normal = FVector::ForwardVector;
+		// Foward direction of the connector.
+		FVector Foward = FVector::ForwardVector;
 
 		// Up direction of the connector.
 		FVector Up = FVector::UpVector;
@@ -232,7 +232,7 @@ protected:
 	bool DuplicatePointForAltDrag(const FVector& InDrag) const;
 
 	/** Split segment using given interpolation value, selects the new point */
-	void SplitSegment(const int32 InSegmentIndex, const float SegmentSplitT) const;
+	void SplitSegment(const int32 InSegmentIndex, const float SegmentSplitT, UZoneShapeComponent* ShapeComp = nullptr) const;
 
 	/** Duplicates selected points and selects them. */
 	void DuplicateSelectedPoints(const FVector& WorldOffset = FVector::ZeroVector, bool bInsertAfter = true) const;
@@ -310,15 +310,45 @@ protected:
 
 	bool bIsSelectingComponent = false;
 
-	bool bIsAutoConnecting = false;
+	/** The point used to calculate the auto connect and intersection states. */
 	int32 SelectedPointForConnecting = -1;
-	TArray<ZoneShapeConnectorRenderInfo> DestShapeConnectorInfos;
-	int32 ClosestShapeConnectorInfoIndex = INDEX_NONE;
-	FVector NearestPointWorldPosition;
-	FVector NearestPointWorldNormal;
+
+	bool bIsAutoConnecting = false;
+	struct FAutoConnectState
+	{
+		TArray<ZoneShapeConnectorRenderInfo> DestShapeConnectorInfos;
+		int32 ClosestShapeConnectorInfoIndex = INDEX_NONE;
+		FVector NearestPointWorldPosition = FVector::ZeroVector;
+		FVector NearestPointWorldNormal = FVector::ZeroVector;
+	};
+	FAutoConnectState AutoConnectState;
+
+	bool bIsCreatingIntersection = false;
+	struct FCreateIntersectionState
+	{
+		TWeakObjectPtr<UZoneShapeComponent> WeakTargetShapeComponent;
+		int32 OverlappingSegmentIndex = INDEX_NONE;
+		float OverlappingSegmentT = -1.0f;
+		int32 ClosePointIndex = INDEX_NONE;
+		FVector PreviewLocation = FVector::ZeroVector;
+	};
+	FCreateIntersectionState CreateIntersectionState;
 
 private:
 
-	void BreakAtPoint(bool bCreateNewActor) const;
+	TArray<UZoneShapeComponent*>  BreakAtPoint(bool bCreateNewActor, UZoneShapeComponent* ShapeComp = nullptr) const;
 	void BreakAtSegment(bool bCreateNewActor) const;
+
+	void DetectCloseByShapeForAutoConnection(const UZoneShapeComponent* ShapeComp, const FZoneShapePoint& DraggedPoint);
+	void DetectCloseByShapeForAutoIntersectionCreation(const UZoneShapeComponent* ShapeComp, const FZoneShapePoint& DraggedPoint);
+
+	void CreateIntersection(UZoneShapeComponent* ShapeComp);
+	void CreateIntersectionForSplineShape(UZoneShapeComponent* ShapeComp, FZoneShapePoint& DraggedPoint, bool DestroyCoveredShape = true);
+	void CreateIntersectionForPolygonShape(UZoneShapeComponent* ShapeComp, FZoneShapePoint& DraggedPoint);
+
+	void ClearAutoConnectingStatus();
+	void ClearAutoIntersectionStatus();
+
+	bool CanAutoConnect(const UZoneShapeComponent* ShapeComp) const;
+	bool CanAutoCreateIntersection(const UZoneShapeComponent* ShapeComp) const;
 };
