ZoneGraph: Fix localization error.

diff --git a/Engine/Plugins/Runtime/ZoneGraph/Source/ZoneGraphEditor/Private/ZoneShapeComponentVisualizer.cpp b/Engine/Plugins/Runtime/ZoneGraph/Source/ZoneGraphEditor/Private/ZoneShapeComponentVisualizer.cpp
--- a/Engine/Plugins/Runtime/ZoneGraph/Source/ZoneGraphEditor/Private/ZoneShapeComponentVisualizer.cpp
+++ b/Engine/Plugins/Runtime/ZoneGraph/Source/ZoneGraphEditor/Private/ZoneShapeComponentVisualizer.cpp
@@ -1574,7 +1574,7 @@ bool FZoneShapeComponentVisualizer::HandleInputKey(FEditorViewportClient* Viewpo
 			if (World->LineTraceSingleByChannel(Hit, Start, End, ECollisionChannel::ECC_WorldStatic, QueryParams))
 			{
 				// Add a new point at the position
-				const FScopedTransaction Transaction(LOCTEXT("AddShapePoint", "Add Shape Point And Snap To Floor"));
+				const FScopedTransaction Transaction(LOCTEXT("AddShapePointAndSnap", "Add Shape Point And Snap To Floor"));
 				ShapeComp->Modify();
 				TArray<FZoneShapePoint>& Points = ShapeComp->GetMutablePoints();
 				FZoneShapePoint PointToAdd(ShapeComp->GetComponentTransform().InverseTransformPosition(Hit.Location));
