Add BezierUtilities back to ZoneGraph area as a deprecated version

diff --git a/Engine/Plugins/Runtime/ZoneGraph/Source/ZoneGraph/Public/BezierUtilities.h b/Engine/Plugins/Runtime/ZoneGraph/Source/ZoneGraph/Public/BezierUtilities.h
new file mode 100644
--- /dev/null
+++ b/Engine/Plugins/Runtime/ZoneGraph/Source/ZoneGraph/Public/BezierUtilities.h
@@ -0,0 +1,9 @@
+// Copyright Epic Games, Inc. All Rights Reserved.
+
+#pragma once
+
+// HEADER_UNIT_SKIP - Deprecated
+
+#include "Curves/BezierUtilities.h"
+
+UE_DEPRECATED_HEADER(5.5, "BezierUtilities.h is deprecated; file moved to Curves/BezierUtilities.h.");
