Zone Graph: Added serialization of Zone Graph Lane Handle.

diff --git a/Engine/Plugins/Runtime/ZoneGraph/Source/ZoneGraph/Public/ZoneGraphTypes.h b/Engine/Plugins/Runtime/ZoneGraph/Source/ZoneGraph/Public/ZoneGraphTypes.h
--- a/Engine/Plugins/Runtime/ZoneGraph/Source/ZoneGraph/Public/ZoneGraphTypes.h
+++ b/Engine/Plugins/Runtime/ZoneGraph/Source/ZoneGraph/Public/ZoneGraphTypes.h
@@ -503,6 +503,14 @@ struct ZONEGRAPH_API FZoneGraphDataHandle
 		Generation = InvalidGeneration;
 	}
 
+	friend FArchive& operator<<(FArchive& Ar, FZoneGraphDataHandle& Handle)
+	{
+		Ar << Handle.Index;
+		Ar << Handle.Generation;
+
+		return Ar;
+	}
+
 	bool IsValid() const { return Generation != InvalidGeneration; }	// Any index is valid, but Generation = 0 means invalid.
 };
 
@@ -541,6 +549,14 @@ struct ZONEGRAPH_API FZoneGraphLaneHandle
 		DataHandle.Reset();
 	}
 
+	friend FArchive& operator<<(FArchive& Ar, FZoneGraphLaneHandle& Handle)
+	{
+		Ar << Handle.Index;
+		Ar << Handle.DataHandle;
+
+		return Ar;
+	}
+
 	FString ToString() const { return FString::Printf(TEXT("[%d/%d]"), DataHandle.Index, Index); }
 	
 	bool IsValid() const { return (Index != INDEX_NONE) && DataHandle.IsValid(); }
