[ZoneGraph] Moving ZoneGraph ShowFlag from the ZoneGraphEditor module to the main ZoneGraph module

diff --git a/Engine/Plugins/Runtime/ZoneGraph/Source/ZoneGraph/Private/ZoneGraphModule.cpp b/Engine/Plugins/Runtime/ZoneGraph/Source/ZoneGraph/Private/ZoneGraphModule.cpp
--- a/Engine/Plugins/Runtime/ZoneGraph/Source/ZoneGraph/Private/ZoneGraphModule.cpp
+++ b/Engine/Plugins/Runtime/ZoneGraph/Source/ZoneGraph/Private/ZoneGraphModule.cpp
@@ -1,9 +1,15 @@
 // Copyright Epic Games, Inc. All Rights Reserved.
 
-#include "CoreMinimal.h"
 #include "Modules/ModuleManager.h"
 #include "IZoneGraphModule.h"
+#include "ShowFlags.h"
 
+#if !UE_BUILD_SHIPPING
+namespace UE::ZoneGraph
+{
+	TCustomShowFlag<> ShowZoneGraph(TEXT("ZoneGraph"), false /*DefaultEnabled*/, SFG_Developer, NSLOCTEXT("ZoneGraphModule", "ShowZoneGraph", "Zone Graph"));
+}
+#endif /* UE_BUILD_SHIPPING */
 
 class FZoneGraphModule : public IZoneGraphModule
 {
diff --git a/Engine/Plugins/Runtime/ZoneGraph/Source/ZoneGraphEditor/Private/ZoneGraphEditorModule.cpp b/Engine/Plugins/Runtime/ZoneGraph/Source/ZoneGraphEditor/Private/ZoneGraphEditorModule.cpp
--- a/Engine/Plugins/Runtime/ZoneGraph/Source/ZoneGraphEditor/Private/ZoneGraphEditorModule.cpp
+++ b/Engine/Plugins/Runtime/ZoneGraph/Source/ZoneGraphEditor/Private/ZoneGraphEditorModule.cpp
@@ -28,12 +28,6 @@
 
 IMPLEMENT_MODULE(FZoneGraphEditorModule, ZoneGraphEditor)
 
-namespace UE::ZoneGraph::Editor
-{
-TCustomShowFlag<> ShowZoneGraph(TEXT("ZoneGraph"), false /*DefaultEnabled*/, SFG_Developer, LOCTEXT("ShowZoneGraph", "Zone Graph"));
-}
-
-
 class FZoneGraphEditorCommands : public TCommands<FZoneGraphEditorCommands>
 {
 public:
