Zone Graph: Adds ability to add zonegraph points at cursor based on selected point

diff --git a/Engine/Plugins/Runtime/ZoneGraph/Source/ZoneGraphEditor/Private/ZoneShapeComponentVisualizer.cpp b/Engine/Plugins/Runtime/ZoneGraph/Source/ZoneGraphEditor/Private/ZoneShapeComponentVisualizer.cpp
--- a/Engine/Plugins/Runtime/ZoneGraph/Source/ZoneGraphEditor/Private/ZoneShapeComponentVisualizer.cpp
+++ b/Engine/Plugins/Runtime/ZoneGraph/Source/ZoneGraphEditor/Private/ZoneShapeComponentVisualizer.cpp
@@ -1576,11 +1576,8 @@ bool FZoneShapeComponentVisualizer::HandleInputKey(FEditorViewportClient* Viewpo
 			{
 				// Add a new point at the position
 				const FScopedTransaction Transaction(LOCTEXT("AddShapePointAndSnap", "Add Shape Point And Snap To Floor"));
-				ShapeComp->Modify();
-				TArray<FZoneShapePoint>& Points = ShapeComp->GetMutablePoints();
-				FZoneShapePoint PointToAdd(ShapeComp->GetComponentTransform().InverseTransformPosition(Hit.Location));
-				Points.Add(PointToAdd);
-				ShapeComp->UpdateShape();
+				const int32 SelectedIndex = SelectionState->GetLastPointIndexSelected();
+				AddSegment(Hit.Location, SelectedIndex, ShapeComp);
 			}
 			else
 			{
@@ -2023,6 +2020,98 @@ void FZoneShapeComponentVisualizer::SplitSegment(const int32 InSegmentIndex, con
 	GEditor->RedrawLevelEditingViewports(true);
 }
 
+void FZoneShapeComponentVisualizer::AddSegment(const FVector& InWorldPos, const int32 InSelectedIndex, UZoneShapeComponent* InShapeComp) const
+{
+	if (!InShapeComp)
+	{
+		InShapeComp = GetEditedShapeComponent();
+	}
+
+	check(InShapeComp != nullptr);
+	check(InSelectedIndex != INDEX_NONE);
+	check(InSelectedIndex >= 0);
+
+	InShapeComp->Modify();
+	if (AActor* Owner = InShapeComp->GetOwner())
+	{
+		Owner->Modify();
+	}
+
+	TArray<FZoneShapePoint>& ShapePoints = InShapeComp->GetMutablePoints();
+	const int32 NumPoints = ShapePoints.Num();
+
+	const int32 PrevPointIdx = InShapeComp->IsShapeClosed() ? (InSelectedIndex + NumPoints - 1) % NumPoints : InSelectedIndex - 1;
+	const int32 NextPointIdx = InShapeComp->IsShapeClosed() ? (InSelectedIndex + 1) % NumPoints : InSelectedIndex + 1;
+
+	FZoneShapePoint NewPoint;
+	NewPoint.Position = InShapeComp->GetComponentTransform().InverseTransformPosition(InWorldPos);
+
+	const FZoneShapePoint& SelectedPoint = ShapePoints[InSelectedIndex];
+	int32 NewPointIndex = InSelectedIndex + 1; // By default, insert new point after the selected
+	
+	if (PrevPointIdx >= 0 && NextPointIdx < NumPoints)
+	{
+		// Both previous and next point are available (selected is neither first nor last)
+		// Calculate which segment is closer
+		const FZoneShapePoint& PrevPoint = ShapePoints[PrevPointIdx];
+		FVector ClosestPointToPrevSegment;
+		float PrevSegmentT = 0.0f;
+		UE::CubicBezier::ClosestPointApproximate(
+			NewPoint.Position,
+			SelectedPoint.Position,
+			SelectedPoint.GetOutControlPoint(),
+			PrevPoint.Position,
+			PrevPoint.GetInControlPoint(),
+			ClosestPointToPrevSegment,
+			PrevSegmentT);
+
+		const FZoneShapePoint& NextPoint = ShapePoints[NextPointIdx];
+		FVector ClosestPointToNextSegment;
+		float NextSegmentT = 0.0f;
+		UE::CubicBezier::ClosestPointApproximate(
+			NewPoint.Position,
+			SelectedPoint.Position,
+			SelectedPoint.GetOutControlPoint(),
+			NextPoint.Position,
+			NextPoint.GetInControlPoint(),
+			ClosestPointToNextSegment,
+			NextSegmentT);
+
+		// Insert new point before the selected if the previous segment is closer
+		if (FVector::Dist(ClosestPointToPrevSegment, NewPoint.Position) < FVector::Dist(ClosestPointToNextSegment, NewPoint.Position))
+		{
+			NewPointIndex = InSelectedIndex;
+		}
+	}
+	else if (PrevPointIdx < 0)
+	{
+		// No previous point (selected is the first) - insert point before selected
+		NewPointIndex = InSelectedIndex;
+	}
+
+	// Copy the type from a selected point if it's a bezier point
+	if (SelectedPoint.Type == FZoneShapePointType::AutoBezier || SelectedPoint.Type == FZoneShapePointType::Bezier)
+	{
+		NewPoint.Type = SelectedPoint.Type;
+	}
+
+	ShapePoints.Insert(NewPoint, NewPointIndex);
+
+	// Create sane default tangent for Bezier points.
+	if (NewPoint.Type == FZoneShapePointType::Bezier)
+	{
+		InShapeComp->UpdatePointRotationAndTangent(NewPointIndex);
+	}
+
+	// Set selection to new point
+	ChangeSelectionState(NewPointIndex, /*bIsCtrlHeld=*/ false);
+
+	InShapeComp->UpdateShape();
+	NotifyPropertyModified(InShapeComp, ShapePointsProperty);
+
+	GEditor->RedrawLevelEditingViewports(/*bInvalidateHitProxies=*/ true);
+}
+
 void FZoneShapeComponentVisualizer::OnDeletePoint() const
 {
 	const FScopedTransaction Transaction(LOCTEXT("DeletePoint", "Delete Points"));
diff --git a/Engine/Plugins/Runtime/ZoneGraph/Source/ZoneGraphEditor/Private/ZoneShapeComponentVisualizer.h b/Engine/Plugins/Runtime/ZoneGraph/Source/ZoneGraphEditor/Private/ZoneShapeComponentVisualizer.h
--- a/Engine/Plugins/Runtime/ZoneGraph/Source/ZoneGraphEditor/Private/ZoneShapeComponentVisualizer.h
+++ b/Engine/Plugins/Runtime/ZoneGraph/Source/ZoneGraphEditor/Private/ZoneShapeComponentVisualizer.h
@@ -234,6 +234,9 @@ protected:
 	/** Split segment using given interpolation value, selects the new point */
 	void SplitSegment(const int32 InSegmentIndex, const float SegmentSplitT, UZoneShapeComponent* ShapeComp = nullptr) const;
 
+	/** Add segment using given position and index, selects the new point */
+	void AddSegment(const FVector& InWorldPos, const int32 InSelectedIndex, UZoneShapeComponent* InShapeComp = nullptr) const;
+
 	/** Duplicates selected points and selects them. */
 	void DuplicateSelectedPoints(const FVector& WorldOffset = FVector::ZeroVector, bool bInsertAfter = true) const;
 
