Updated headers using UnrealCodeFixup to make sure dllstorage is on methods/staticvars instead of types. Part 4/n

diff --git a/Engine/Plugins/Runtime/USDCore/Source/USDClasses/Public/USDAssetUserData.h b/Engine/Plugins/Runtime/USDCore/Source/USDClasses/Public/USDAssetUserData.h
--- a/Engine/Plugins/Runtime/USDCore/Source/USDClasses/Public/USDAssetUserData.h
+++ b/Engine/Plugins/Runtime/USDCore/Source/USDClasses/Public/USDAssetUserData.h
@@ -8,8 +8,8 @@
 
 #include "USDAssetUserData.generated.h"
 
-UCLASS(BlueprintType)
-class USDCLASSES_API UUsdAssetUserData : public UAssetUserData
+UCLASS(MinimalAPI, BlueprintType)
+class UUsdAssetUserData : public UAssetUserData
 {
 	GENERATED_BODY()
 
@@ -30,8 +30,8 @@ public:
 	FString OriginalHash;
 };
 
-UCLASS(BlueprintType)
-class USDCLASSES_API UUsdAnimSequenceAssetUserData : public UUsdAssetUserData
+UCLASS(MinimalAPI, BlueprintType)
+class UUsdAnimSequenceAssetUserData : public UUsdAssetUserData
 {
 	GENERATED_BODY()
 
@@ -58,8 +58,8 @@ struct FUsdPrimPathList
 };
 
 /** AssetImportData assigned to Unreal materials that are generated when parsing USD Material prims */
-UCLASS(BlueprintType)
-class USDCLASSES_API UUsdMaterialAssetUserData : public UUsdAssetUserData
+UCLASS(MinimalAPI, BlueprintType)
+class UUsdMaterialAssetUserData : public UUsdAssetUserData
 {
 	GENERATED_BODY()
 
@@ -86,8 +86,8 @@ public:
 };
 
 /** We assign these to UStaticMeshes or USkeletalMeshes generated from USD */
-UCLASS(BlueprintType)
-class USDCLASSES_API UUsdMeshAssetUserData : public UUsdAssetUserData
+UCLASS(MinimalAPI, BlueprintType)
+class UUsdMeshAssetUserData : public UUsdAssetUserData
 {
 	GENERATED_BODY()
 
@@ -106,8 +106,8 @@ public:
 };
 
 /** We assign these to UGeometryCaches generated from USD */
-UCLASS(BlueprintType)
-class USDCLASSES_API UUsdGeometryCacheAssetUserData : public UUsdMeshAssetUserData
+UCLASS(MinimalAPI, BlueprintType)
+class UUsdGeometryCacheAssetUserData : public UUsdMeshAssetUserData
 {
 	GENERATED_BODY()
 
@@ -146,8 +146,8 @@ public:
 /**
  * We use this mainly to help in mapping between stage timeCode and the FrameIndex for animated SVTs
  */
-UCLASS(BlueprintType)
-class USDCLASSES_API UUsdSparseVolumeTextureAssetUserData : public UUsdAssetUserData
+UCLASS(MinimalAPI, BlueprintType)
+class UUsdSparseVolumeTextureAssetUserData : public UUsdAssetUserData
 {
 	GENERATED_BODY()
 
diff --git a/Engine/Plugins/Runtime/USDCore/Source/USDClasses/Public/USDDrawModeComponent.h b/Engine/Plugins/Runtime/USDCore/Source/USDClasses/Public/USDDrawModeComponent.h
--- a/Engine/Plugins/Runtime/USDCore/Source/USDClasses/Public/USDDrawModeComponent.h
+++ b/Engine/Plugins/Runtime/USDCore/Source/USDClasses/Public/USDDrawModeComponent.h
@@ -6,6 +6,8 @@
 
 #include "USDDrawModeComponent.generated.h"
 
+#define UE_API USDCLASSES_API
+
 class UMaterialInstance;
 enum class EUsdUpAxis : uint8;
 
@@ -55,8 +57,8 @@ namespace UsdUtils
  * Component type that is used to draw bounds, cards and origin axes when translating a prim that has the UsdGeomModelAPI schema
  * applied.
  */
-UCLASS(ClassGroup = (USD), meta = (BlueprintSpawnableComponent))
-class USDCLASSES_API UUsdDrawModeComponent : public UPrimitiveComponent
+UCLASS(MinimalAPI, ClassGroup = (USD), meta = (BlueprintSpawnableComponent))
+class UUsdDrawModeComponent : public UPrimitiveComponent
 {
 	GENERATED_BODY()
 
@@ -112,62 +114,64 @@ private:
 
 public:
 	UFUNCTION(BlueprintSetter, Category = "USD")
-	void SetBoundsMin(const FVector& NewMin);
+	UE_API void SetBoundsMin(const FVector& NewMin);
 
 	UFUNCTION(BlueprintSetter, Category = "USD")
-	void SetBoundsMax(const FVector& NewMax);
+	UE_API void SetBoundsMax(const FVector& NewMax);
 
 	UFUNCTION(BlueprintSetter, Category = "USD")
-	void SetDrawMode(EUsdDrawMode NewDrawMode);
+	UE_API void SetDrawMode(EUsdDrawMode NewDrawMode);
 
 	UFUNCTION(BlueprintSetter, Category = "USD")
-	void SetBoundsColor(FLinearColor NewColor);
+	UE_API void SetBoundsColor(FLinearColor NewColor);
 
 	UFUNCTION(BlueprintSetter, Category = "USD")
-	void SetCardGeometry(EUsdModelCardGeometry NewGeometry);
+	UE_API void SetCardGeometry(EUsdModelCardGeometry NewGeometry);
 
 	UFUNCTION(BlueprintSetter, Category = "USD|Cards")
-	void SetCardTextureXPos(UTexture2D* NewTexture);
+	UE_API void SetCardTextureXPos(UTexture2D* NewTexture);
 
 	UFUNCTION(BlueprintSetter, Category = "USD|Cards")
-	void SetCardTextureYPos(UTexture2D* NewTexture);
+	UE_API void SetCardTextureYPos(UTexture2D* NewTexture);
 
 	UFUNCTION(BlueprintSetter, Category = "USD|Cards")
-	void SetCardTextureZPos(UTexture2D* NewTexture);
+	UE_API void SetCardTextureZPos(UTexture2D* NewTexture);
 
 	UFUNCTION(BlueprintSetter, Category = "USD|Cards")
-	void SetCardTextureXNeg(UTexture2D* NewTexture);
+	UE_API void SetCardTextureXNeg(UTexture2D* NewTexture);
 
 	UFUNCTION(BlueprintSetter, Category = "USD|Cards")
-	void SetCardTextureYNeg(UTexture2D* NewTexture);
+	UE_API void SetCardTextureYNeg(UTexture2D* NewTexture);
 
 	UFUNCTION(BlueprintSetter, Category = "USD|Cards")
-	void SetCardTextureZNeg(UTexture2D* NewTexture);
+	UE_API void SetCardTextureZNeg(UTexture2D* NewTexture);
 
 	UFUNCTION(BlueprintPure, Category = "USD|Cards")
-	UTexture2D* GetTextureForFace(EUsdModelCardFace Face) const;
+	UE_API UTexture2D* GetTextureForFace(EUsdModelCardFace Face) const;
 
 	UFUNCTION(BlueprintCallable, Category = "USD|Cards")
-	void SetTextureForFace(EUsdModelCardFace Face, UTexture2D* Texture);
+	UE_API void SetTextureForFace(EUsdModelCardFace Face, UTexture2D* Texture);
 
 public:
-	UUsdDrawModeComponent();
+	UE_API UUsdDrawModeComponent();
 	virtual ~UUsdDrawModeComponent() override{};
 
 	// Begin UPrimitiveComponent interface
-	virtual FPrimitiveSceneProxy* CreateSceneProxy() override;
-	virtual FBoxSphereBounds CalcBounds(const FTransform& LocalToWorld) const override;
-	virtual void GetUsedMaterials(TArray<UMaterialInterface*>& OutMaterials, bool bGetDebugMaterials = false) const override;
+	UE_API virtual FPrimitiveSceneProxy* CreateSceneProxy() override;
+	UE_API virtual FBoxSphereBounds CalcBounds(const FTransform& LocalToWorld) const override;
+	UE_API virtual void GetUsedMaterials(TArray<UMaterialInterface*>& OutMaterials, bool bGetDebugMaterials = false) const override;
 #if WITH_EDITOR
-	virtual void PostEditChangeProperty(FPropertyChangedEvent& PropertyChangedEvent) override;
-	virtual HHitProxy* CreateMeshHitProxy(int32 SectionIndex, int32 MaterialIndex) const override;
+	UE_API virtual void PostEditChangeProperty(FPropertyChangedEvent& PropertyChangedEvent) override;
+	UE_API virtual HHitProxy* CreateMeshHitProxy(int32 SectionIndex, int32 MaterialIndex) const override;
 #endif
 	// End UPrimitiveComponent interface
 
-	EUsdModelCardFace GetAuthoredFaces() const;
-	void SetAuthoredFaces(EUsdModelCardFace NewAuthoredFaces);
+	UE_API EUsdModelCardFace GetAuthoredFaces() const;
+	UE_API void SetAuthoredFaces(EUsdModelCardFace NewAuthoredFaces);
 
 private:
-	UMaterialInstance* GetOrCreateTextureMaterial(EUsdModelCardFace Face);
-	void RefreshMaterialInstances();
+	UE_API UMaterialInstance* GetOrCreateTextureMaterial(EUsdModelCardFace Face);
+	UE_API void RefreshMaterialInstances();
 };
+
+#undef UE_API
diff --git a/Engine/Plugins/Runtime/USDCore/Source/USDClasses/Public/USDMaterialUtils.h b/Engine/Plugins/Runtime/USDCore/Source/USDClasses/Public/USDMaterialUtils.h
--- a/Engine/Plugins/Runtime/USDCore/Source/USDClasses/Public/USDMaterialUtils.h
+++ b/Engine/Plugins/Runtime/USDCore/Source/USDClasses/Public/USDMaterialUtils.h
@@ -4,6 +4,8 @@
 
 #include "CoreMinimal.h"
 
+#define UE_API USDCLASSES_API
+
 class UMaterialInstanceDynamic;
 class UMaterialInstanceConstant;
 
@@ -19,14 +21,14 @@ ENUM_CLASS_FLAGS(EUsdReferenceMaterialProperties)
 namespace UsdUnreal::MaterialUtils
 {
 	/** Describes the type of vertex color/DisplayColor material that we would need in order to render a prim's displayColor data as intended */
-	struct USDCLASSES_API FDisplayColorMaterial
+	struct FDisplayColorMaterial
 	{
 		bool bHasOpacity = false;
 		bool bIsDoubleSided = false;
 
-		FString ToString() const;
-		FString ToPrettyString() const;
-		static TOptional<FDisplayColorMaterial> FromString(const FString& DisplayColorString);
+		UE_API FString ToString() const;
+		UE_API FString ToPrettyString() const;
+		static UE_API TOptional<FDisplayColorMaterial> FromString(const FString& DisplayColorString);
 	};
 
 	USDCLASSES_API const FSoftObjectPath* GetReferenceMaterialPath(const FDisplayColorMaterial& DisplayColorDescription);
@@ -62,3 +64,5 @@ namespace UsdUnreal::MaterialUtils
 	USDCLASSES_API void UnregisterRenderContext(const FName& RenderContextName);
 	USDCLASSES_API const TArray<FName>& GetRegisteredRenderContexts();
 };	  // namespace UsdUnreal::MaterialUtils
+
+#undef UE_API
diff --git a/Engine/Plugins/Runtime/USDCore/Source/USDClasses/Public/USDMetadataImportOptions.h b/Engine/Plugins/Runtime/USDCore/Source/USDClasses/Public/USDMetadataImportOptions.h
--- a/Engine/Plugins/Runtime/USDCore/Source/USDClasses/Public/USDMetadataImportOptions.h
+++ b/Engine/Plugins/Runtime/USDCore/Source/USDClasses/Public/USDMetadataImportOptions.h
@@ -7,7 +7,7 @@
 struct FAnalyticsEventAttribute;
 
 USTRUCT(BlueprintType)
-struct USDCLASSES_API FUsdMetadataImportOptions
+struct FUsdMetadataImportOptions
 {
 	GENERATED_BODY()
 
diff --git a/Engine/Plugins/Runtime/USDCore/Source/USDClasses/Public/USDReferenceOptions.h b/Engine/Plugins/Runtime/USDCore/Source/USDClasses/Public/USDReferenceOptions.h
--- a/Engine/Plugins/Runtime/USDCore/Source/USDClasses/Public/USDReferenceOptions.h
+++ b/Engine/Plugins/Runtime/USDCore/Source/USDClasses/Public/USDReferenceOptions.h
@@ -7,8 +7,8 @@
 /**
  * Options to display when adding a reference or a payload for a prim
  */
-UCLASS(Blueprintable)
-class USDCLASSES_API UUsdReferenceOptions : public UObject
+UCLASS(MinimalAPI, Blueprintable)
+class UUsdReferenceOptions : public UObject
 {
 	GENERATED_BODY()
 
diff --git a/Engine/Plugins/Runtime/USDCore/Source/USDClasses/Public/USDStageOptions.h b/Engine/Plugins/Runtime/USDCore/Source/USDClasses/Public/USDStageOptions.h
--- a/Engine/Plugins/Runtime/USDCore/Source/USDClasses/Public/USDStageOptions.h
+++ b/Engine/Plugins/Runtime/USDCore/Source/USDClasses/Public/USDStageOptions.h
@@ -14,7 +14,7 @@ enum class EUsdUpAxis : uint8
 };
 
 USTRUCT(BlueprintType)
-struct USDCLASSES_API FUsdStageOptions
+struct FUsdStageOptions
 {
 	GENERATED_BODY()
 
diff --git a/Engine/Plugins/Runtime/USDCore/Source/USDClasses/Public/USDUnrealAssetInfo.h b/Engine/Plugins/Runtime/USDCore/Source/USDClasses/Public/USDUnrealAssetInfo.h
--- a/Engine/Plugins/Runtime/USDCore/Source/USDClasses/Public/USDUnrealAssetInfo.h
+++ b/Engine/Plugins/Runtime/USDCore/Source/USDClasses/Public/USDUnrealAssetInfo.h
@@ -8,7 +8,7 @@
  * Metadata added to a prim to indicate it was exported from a particular Unreal asset
  */
 USTRUCT(BlueprintType)
-struct USDCLASSES_API FUsdUnrealAssetInfo
+struct FUsdUnrealAssetInfo
 {
 	GENERATED_BODY()
 
diff --git a/Engine/Plugins/Runtime/USDCore/Source/USDUtilities/Public/Objects/USDInfoCache.h b/Engine/Plugins/Runtime/USDCore/Source/USDUtilities/Public/Objects/USDInfoCache.h
--- a/Engine/Plugins/Runtime/USDCore/Source/USDUtilities/Public/Objects/USDInfoCache.h
+++ b/Engine/Plugins/Runtime/USDCore/Source/USDUtilities/Public/Objects/USDInfoCache.h
@@ -8,6 +8,8 @@
 
 #include "UsdWrappers/SdfPath.h"
 
+#define UE_API USDUTILITIES_API
+
 struct FUsdSchemaTranslationContext;
 namespace UE
 {
@@ -31,92 +33,92 @@ struct FUsdInfoCacheImpl;
 /**
  * Caches information about a specific USD Stage
  */
-class USDUTILITIES_API FUsdInfoCache
+class FUsdInfoCache
 {
 public:
-	FUsdInfoCache();
-	virtual ~FUsdInfoCache();
+	UE_API FUsdInfoCache();
+	UE_API virtual ~FUsdInfoCache();
 
-	void CopyImpl(const FUsdInfoCache& Other);
+	UE_API void CopyImpl(const FUsdInfoCache& Other);
 
-	bool Serialize(FArchive& Ar);
+	UE_API bool Serialize(FArchive& Ar);
 
 	// Returns whether we contain any info about prim at 'Path' at all
-	bool ContainsInfoAboutPrim(const UE::FSdfPath& Path) const;
+	UE_API bool ContainsInfoAboutPrim(const UE::FSdfPath& Path) const;
 
 	// Retrieves the children of a prim from the cached information
-	TArray<UE::FSdfPath> GetChildren(const UE::FSdfPath& ParentPath) const;
+	UE_API TArray<UE::FSdfPath> GetChildren(const UE::FSdfPath& ParentPath) const;
 
 	// Returns a list of all prims we have generic info about
 	UE_DEPRECATED(5.5, "No longer used")
-	TSet<UE::FSdfPath> GetKnownPrims() const;
+	UE_API TSet<UE::FSdfPath> GetKnownPrims() const;
 
-	void RebuildCacheForSubtree(const UE::FUsdPrim& Prim, FUsdSchemaTranslationContext& Context);
-	void RebuildCacheForSubtrees(const TArray<UE::FSdfPath>& SubtreeRoots, FUsdSchemaTranslationContext& Context);
+	UE_API void RebuildCacheForSubtree(const UE::FUsdPrim& Prim, FUsdSchemaTranslationContext& Context);
+	UE_API void RebuildCacheForSubtrees(const TArray<UE::FSdfPath>& SubtreeRoots, FUsdSchemaTranslationContext& Context);
 
-	void Clear();
-	bool IsEmpty();
+	UE_API void Clear();
+	UE_API bool IsEmpty();
 
 public:
-	bool IsPathCollapsed(const UE::FSdfPath& Path, ECollapsingType CollapsingType) const;
-	bool DoesPathCollapseChildren(const UE::FSdfPath& Path, ECollapsingType CollapsingType) const;
+	UE_API bool IsPathCollapsed(const UE::FSdfPath& Path, ECollapsingType CollapsingType) const;
+	UE_API bool DoesPathCollapseChildren(const UE::FSdfPath& Path, ECollapsingType CollapsingType) const;
 
 	// Returns Path in case it represents an uncollapsed prim, or returns the path to the prim that collapsed it
-	UE::FSdfPath UnwindToNonCollapsedPath(const UE::FSdfPath& Path, ECollapsingType CollapsingType) const;
+	UE_API UE::FSdfPath UnwindToNonCollapsedPath(const UE::FSdfPath& Path, ECollapsingType CollapsingType) const;
 
 public:
 	// Returns the paths to prims that, when translated into assets or components, also require reading the prim at
 	// 'Path'. e.g. providing the path to a Shader prim will return the paths to all Material prims for which the
 	// translation involves reading that particular Shader.
-	TSet<UE::FSdfPath> GetMainPrims(const UE::FSdfPath& AuxPrimPath) const;
+	UE_API TSet<UE::FSdfPath> GetMainPrims(const UE::FSdfPath& AuxPrimPath) const;
 
 	// The inverse of the function above: Provide it with the path to a Material prim and it will return the set of
 	// paths to all Shader prims that need to be read to translate that Material prim into material assets
-	TSet<UE::FSdfPath> GetAuxiliaryPrims(const UE::FSdfPath& MainPrimPath) const;
+	UE_API TSet<UE::FSdfPath> GetAuxiliaryPrims(const UE::FSdfPath& MainPrimPath) const;
 
 public:
-	TSet<UE::FSdfPath> GetMaterialUsers(const UE::FSdfPath& Path) const;
-	bool IsMaterialUsed(const UE::FSdfPath& Path) const;
+	UE_API TSet<UE::FSdfPath> GetMaterialUsers(const UE::FSdfPath& Path) const;
+	UE_API bool IsMaterialUsed(const UE::FSdfPath& Path) const;
 
 public:
 	// Provides the total vertex or material slots counts for each prim *and* its subtree.
 	// This is built inside RebuildCacheForSubtree, so it will factor in the used Context's bMergeIdenticalMaterialSlots.
 	// Note that these aren't affected by actual collapsing: A prim that doesn't collapse its children will still
 	// provide the total sum of vertex counts of its entire subtree when queried
-	TOptional<uint64> GetSubtreeVertexCount(const UE::FSdfPath& Path);
-	TOptional<uint64> GetSubtreeMaterialSlotCount(const UE::FSdfPath& Path);
-	TOptional<TArray<UsdUtils::FUsdPrimMaterialSlot>> GetSubtreeMaterialSlots(const UE::FSdfPath& Path);
+	UE_API TOptional<uint64> GetSubtreeVertexCount(const UE::FSdfPath& Path);
+	UE_API TOptional<uint64> GetSubtreeMaterialSlotCount(const UE::FSdfPath& Path);
+	UE_API TOptional<TArray<UsdUtils::FUsdPrimMaterialSlot>> GetSubtreeMaterialSlots(const UE::FSdfPath& Path);
 
 	// Returns true if Path could potentially be collapsed as a Geometry Cache asset
 	UE_DEPRECATED(5.5, "No longer used")
-	bool IsPotentialGeometryCacheRoot(const UE::FSdfPath& Path) const;
+	UE_API bool IsPotentialGeometryCacheRoot(const UE::FSdfPath& Path) const;
 
 public:
 	// Marks/checks if the provided path to a prototype prim is already being translated.
 	// This is used during scene translation with instanceables, so that the schema translators can early out
 	// in case they have been created to translate multiple instances of the same prototype
-	void ResetTranslatedPrototypes();
-	bool IsPrototypeTranslated(const UE::FSdfPath& PrototypePath);
-	void MarkPrototypeAsTranslated(const UE::FSdfPath& PrototypePath);
+	UE_API void ResetTranslatedPrototypes();
+	UE_API bool IsPrototypeTranslated(const UE::FSdfPath& PrototypePath);
+	UE_API void MarkPrototypeAsTranslated(const UE::FSdfPath& PrototypePath);
 
 public:
 	UE_DEPRECATED(5.5, "Use the UUsdPrimLinkCache object and its analogous function instead")
-	void LinkAssetToPrim(const UE::FSdfPath& Path, UObject* Asset);
+	UE_API void LinkAssetToPrim(const UE::FSdfPath& Path, UObject* Asset);
 
 	UE_DEPRECATED(5.5, "Use the UUsdPrimLinkCache object and its analogous function instead")
-	void UnlinkAssetFromPrim(const UE::FSdfPath& Path, UObject* Asset);
+	UE_API void UnlinkAssetFromPrim(const UE::FSdfPath& Path, UObject* Asset);
 
 	UE_DEPRECATED(5.5, "Use the UUsdPrimLinkCache object and its analogous function instead")
-	TArray<TWeakObjectPtr<UObject>> RemoveAllAssetPrimLinks(const UE::FSdfPath& Path);
+	UE_API TArray<TWeakObjectPtr<UObject>> RemoveAllAssetPrimLinks(const UE::FSdfPath& Path);
 
 	UE_DEPRECATED(5.5, "Use the UUsdPrimLinkCache object and its analogous function instead")
-	TArray<UE::FSdfPath> RemoveAllAssetPrimLinks(const UObject* Asset);
+	UE_API TArray<UE::FSdfPath> RemoveAllAssetPrimLinks(const UObject* Asset);
 
 	UE_DEPRECATED(5.5, "Use the UUsdPrimLinkCache object and its analogous function instead")
-	void RemoveAllAssetPrimLinks();
+	UE_API void RemoveAllAssetPrimLinks();
 
 	UE_DEPRECATED(5.5, "Use the UUsdPrimLinkCache object and its analogous function instead")
-	TArray<TWeakObjectPtr<UObject>> GetAllAssetsForPrim(const UE::FSdfPath& Path) const;
+	UE_API TArray<TWeakObjectPtr<UObject>> GetAllAssetsForPrim(const UE::FSdfPath& Path) const;
 
 	template<typename T = UObject>
 	UE_DEPRECATED(5.5, "Use the UUsdPrimLinkCache object and its analogous function instead")
@@ -133,10 +135,10 @@ public:
 	}
 
 	UE_DEPRECATED(5.5, "Use the UUsdPrimLinkCache object and its analogous function instead")
-	TArray<UE::FSdfPath> GetPrimsForAsset(const UObject* Asset) const;
+	UE_API TArray<UE::FSdfPath> GetPrimsForAsset(const UObject* Asset) const;
 
 	UE_DEPRECATED(5.5, "Use the UUsdPrimLinkCache object and its analogous function instead")
-	TMap<UE::FSdfPath, TArray<TWeakObjectPtr<UObject>>> GetAllAssetPrimLinks() const;
+	UE_API TMap<UE::FSdfPath, TArray<TWeakObjectPtr<UObject>>> GetAllAssetPrimLinks() const;
 
 private:
 	friend class FUsdGeomXformableTranslator;
@@ -150,12 +152,14 @@ private:
 	// meaning it is unknown at this point whether the prim CanBeCollapsed or not.
 	//
 	// In general, you shouldn't call this, but just use "IsPathCollapsed" or "DoesPathCollapseChildren" instead.
-	TOptional<bool> CanXformableSubtreeBeCollapsed(const UE::FSdfPath& RootPath, FUsdSchemaTranslationContext& Context) const;
+	UE_API TOptional<bool> CanXformableSubtreeBeCollapsed(const UE::FSdfPath& RootPath, FUsdSchemaTranslationContext& Context) const;
 
 	// Analogous to the function above, this overload of IsPotentialGeometryCacheRoot is meant for internal use, and exists because
 	// during the info cache build (in some contexts) we can fill in this geometry cache information on-demand, for better performance.
-	bool IsPotentialGeometryCacheRoot(const UE::FUsdPrim& Prim) const;
+	UE_API bool IsPotentialGeometryCacheRoot(const UE::FUsdPrim& Prim) const;
 
 private:
 	TUniquePtr<FUsdInfoCacheImpl> Impl;
 };
+
+#undef UE_API
diff --git a/Engine/Plugins/Runtime/USDCore/Source/USDUtilities/Public/Objects/USDPrimLinkCache.h b/Engine/Plugins/Runtime/USDCore/Source/USDUtilities/Public/Objects/USDPrimLinkCache.h
--- a/Engine/Plugins/Runtime/USDCore/Source/USDUtilities/Public/Objects/USDPrimLinkCache.h
+++ b/Engine/Plugins/Runtime/USDCore/Source/USDUtilities/Public/Objects/USDPrimLinkCache.h
@@ -8,38 +8,40 @@
 
 #include "UsdWrappers/SdfPath.h"
 
+#define UE_API USDUTILITIES_API
+
 namespace UE
 {
 	class FSdfPath;
 }
 
-class USDUTILITIES_API FUsdPrimLinkCache
+class FUsdPrimLinkCache
 {
 public:
 	struct FUsdPrimLinkCacheImpl;
 
-	FUsdPrimLinkCache();
-	~FUsdPrimLinkCache();
+	UE_API FUsdPrimLinkCache();
+	UE_API ~FUsdPrimLinkCache();
 
 	// Begin UObject interface
-	void Serialize(FArchive& Ar);
+	UE_API void Serialize(FArchive& Ar);
 	// End UObject interface
 
 	// Returns whether we contain any info about prim at 'Path' at all
-	bool ContainsInfoAboutPrim(const UE::FSdfPath& Path) const;
+	UE_API bool ContainsInfoAboutPrim(const UE::FSdfPath& Path) const;
 
-	void Clear();
-	bool IsEmpty();
+	UE_API void Clear();
+	UE_API bool IsEmpty();
 
 public:
-	void LinkAssetToPrim(const UE::FSdfPath& Path, UObject* Asset);
-	void UnlinkAssetFromPrim(const UE::FSdfPath& Path, UObject* Asset);
+	UE_API void LinkAssetToPrim(const UE::FSdfPath& Path, UObject* Asset);
+	UE_API void UnlinkAssetFromPrim(const UE::FSdfPath& Path, UObject* Asset);
 
-	TArray<TWeakObjectPtr<UObject>> RemoveAllAssetPrimLinks(const UE::FSdfPath& Path);
-	TArray<UE::FSdfPath> RemoveAllAssetPrimLinks(const UObject* Asset);
-	void RemoveAllAssetPrimLinks();
+	UE_API TArray<TWeakObjectPtr<UObject>> RemoveAllAssetPrimLinks(const UE::FSdfPath& Path);
+	UE_API TArray<UE::FSdfPath> RemoveAllAssetPrimLinks(const UObject* Asset);
+	UE_API void RemoveAllAssetPrimLinks();
 
-	TArray<TWeakObjectPtr<UObject>> GetAllAssetsForPrim(const UE::FSdfPath& Path) const;
+	UE_API TArray<TWeakObjectPtr<UObject>> GetAllAssetsForPrim(const UE::FSdfPath& Path) const;
 
 	template<typename T = UObject>
 	T* GetSingleAssetForPrim(const UE::FSdfPath& Path) const
@@ -78,9 +80,11 @@ public:
 		return CastAssets;
 	}
 
-	TArray<UE::FSdfPath> GetPrimsForAsset(const UObject* Asset) const;
-	TMap<UE::FSdfPath, TArray<TWeakObjectPtr<UObject>>> GetAllAssetPrimLinks() const;
+	UE_API TArray<UE::FSdfPath> GetPrimsForAsset(const UObject* Asset) const;
+	UE_API TMap<UE::FSdfPath, TArray<TWeakObjectPtr<UObject>>> GetAllAssetPrimLinks() const;
 
 private:
 	TUniquePtr<FUsdPrimLinkCacheImpl> Impl;
 };
+
+#undef UE_API
diff --git a/Engine/Plugins/Runtime/USDCore/Source/USDUtilities/Public/Objects/USDSchemaTranslator.h b/Engine/Plugins/Runtime/USDCore/Source/USDUtilities/Public/Objects/USDSchemaTranslator.h
--- a/Engine/Plugins/Runtime/USDCore/Source/USDUtilities/Public/Objects/USDSchemaTranslator.h
+++ b/Engine/Plugins/Runtime/USDCore/Source/USDUtilities/Public/Objects/USDSchemaTranslator.h
@@ -43,7 +43,7 @@ namespace UE
 	class FUsdGeomBBoxCache;
 }
 
-class USDUTILITIES_API FRegisteredSchemaTranslatorHandle
+class FRegisteredSchemaTranslatorHandle
 {
 public:
 	FRegisteredSchemaTranslatorHandle()
@@ -76,24 +76,24 @@ public:
 	}
 
 private:
-	static int32 CurrentSchemaTranslatorId;
+	static USDUTILITIES_API int32 CurrentSchemaTranslatorId;
 
 	FString SchemaName;
 	int32 Id;
 };
 
-class USDUTILITIES_API FUsdSchemaTranslatorRegistry
+class FUsdSchemaTranslatorRegistry
 {
 	using FCreateTranslator = TFunction<TSharedRef<FUsdSchemaTranslator>(TSharedRef<FUsdSchemaTranslationContext>, const UE::FUsdTyped&)>;
 	using FSchemaTranslatorsStack = TArray<FRegisteredSchemaTranslator, TInlineAllocator<1>>;
 
 public:
-	static FUsdSchemaTranslatorRegistry& Get();
+	static USDUTILITIES_API FUsdSchemaTranslatorRegistry& Get();
 
 	/**
 	 * Returns the translator to use for InSchema
 	 */
-	TSharedPtr<FUsdSchemaTranslator> CreateTranslatorForSchema(
+	USDUTILITIES_API TSharedPtr<FUsdSchemaTranslator> CreateTranslatorForSchema(
 		TSharedRef<FUsdSchemaTranslationContext> InTranslationContext,
 		const UE::FUsdTyped& InSchema
 	);
@@ -115,21 +115,21 @@ public:
 		return Register(SchemaName, CreateSchemaTranslator);
 	}
 
-	void Unregister(const FRegisteredSchemaTranslatorHandle& TranslatorHandle);
+	USDUTILITIES_API void Unregister(const FRegisteredSchemaTranslatorHandle& TranslatorHandle);
 
-	int32 GetExternalSchemaTranslatorCount();
+	USDUTILITIES_API int32 GetExternalSchemaTranslatorCount();
 
 protected:
-	FRegisteredSchemaTranslatorHandle Register(const FString& SchemaName, FCreateTranslator CreateFunction);
+	USDUTILITIES_API FRegisteredSchemaTranslatorHandle Register(const FString& SchemaName, FCreateTranslator CreateFunction);
 
-	FSchemaTranslatorsStack* FindSchemaTranslatorStack(const FString& SchemaName);
+	USDUTILITIES_API FSchemaTranslatorsStack* FindSchemaTranslatorStack(const FString& SchemaName);
 
 	TArray<TPair<FString, FSchemaTranslatorsStack>> RegisteredSchemaTranslators;
 
 private:
 	// Small machinery that lets us collect basic analytics about how many custom schema translators are being used in this session
 	friend class FUsdSchemasModule;
-	void ResetExternalTranslatorCount();
+	USDUTILITIES_API void ResetExternalTranslatorCount();
 	int32 ExternalSchemaTranslatorCount = 0;
 };
 
@@ -154,7 +154,7 @@ public:
 	const FName& GetUnrealRenderContext() const;
 };
 
-struct USDUTILITIES_API FUsdSchemaTranslationContext : public TSharedFromThis<FUsdSchemaTranslationContext>
+struct FUsdSchemaTranslationContext : public TSharedFromThis<FUsdSchemaTranslationContext>
 {
 	// Explicitly declare these defaulted special functions or else the compiler will do it elsewhere and emit
 	// deprecated warnings due to usage of bCollapseTopLevelPointInstancers
@@ -164,10 +164,10 @@ struct USDUTILITIES_API FUsdSchemaTranslationContext : public TSharedFromThis<FU
 	PRAGMA_ENABLE_DEPRECATION_WARNINGS
 
 	UE_DEPRECATED(5.5, "Use the constructor that receives an UUsdAssetCache3 instead")
-	explicit FUsdSchemaTranslationContext(const UE::FUsdStage& InStage, UUsdAssetCache2& InAssetCache);
+	USDUTILITIES_API explicit FUsdSchemaTranslationContext(const UE::FUsdStage& InStage, UUsdAssetCache2& InAssetCache);
 
-	explicit FUsdSchemaTranslationContext(const UE::FUsdStage& InStage);
-	explicit FUsdSchemaTranslationContext(const UE::FUsdStage& InStage, UUsdAssetCache3& InAssetCache);
+	USDUTILITIES_API explicit FUsdSchemaTranslationContext(const UE::FUsdStage& InStage);
+	USDUTILITIES_API explicit FUsdSchemaTranslationContext(const UE::FUsdStage& InStage, UUsdAssetCache3& InAssetCache);
 
 	/** True if we're a context created by the USDStageImporter to fully import to persistent assets and actors */
 	bool bIsImporting = false;
@@ -330,7 +330,7 @@ struct USDUTILITIES_API FUsdSchemaTranslationContext : public TSharedFromThis<FU
 		return Level != nullptr;
 	}
 
-	void CompleteTasks();
+	USDUTILITIES_API void CompleteTasks();
 
 	TArray<TSharedPtr<FUsdSchemaTranslatorTaskChain>> TranslatorTasks;
 };
@@ -357,7 +357,7 @@ enum class ESchemaTranslationLaunchPolicy
 	Async
 };
 
-class USDUTILITIES_API FUsdSchemaTranslator
+class FUsdSchemaTranslator
 {
 public:
 	explicit FUsdSchemaTranslator(TSharedRef<FUsdSchemaTranslationContext> InContext, const UE::FUsdTyped& InSchema)
@@ -395,7 +395,7 @@ public:
 		return {};
 	}
 
-	bool IsCollapsed(ECollapsingType CollapsingType) const;
+	USDUTILITIES_API bool IsCollapsed(ECollapsingType CollapsingType) const;
 	virtual bool CanBeCollapsed(ECollapsingType CollapsingType) const
 	{
 		return false;
@@ -412,14 +412,14 @@ public:
 	 * will "own" the translation for that prototype, and any subsequent calls by other schema translators
 	 * with the same prototype will just return true so they can early out.
 	 */
-	bool ShouldSkipInstance() const;
+	USDUTILITIES_API bool ShouldSkipInstance() const;
 
 	/**
 	 * If this prim is a prototype or an instance proxy, returns the prototype path (or the path to the
 	 * analogue prim in the prototype's hierarchy).
 	 * If this prim is just a regular non-instance prim, this just returns our PrimPath member.
 	 */
-	UE::FSdfPath GetPrototypePrimPath() const;
+	USDUTILITIES_API UE::FSdfPath GetPrototypePrimPath() const;
 
 	UE::FUsdPrim GetPrim() const
 	{
@@ -431,7 +431,7 @@ protected:
 	TSharedRef<FUsdSchemaTranslationContext> Context;
 };
 
-struct USDUTILITIES_API FSchemaTranslatorTask
+struct FSchemaTranslatorTask
 {
 	explicit FSchemaTranslatorTask(ESchemaTranslationLaunchPolicy InPolicy, TFunction<bool()> InCallable)
 		: Callable(InCallable)
@@ -449,15 +449,15 @@ struct USDUTILITIES_API FSchemaTranslatorTask
 
 	FThreadSafeBool bIsDone;
 
-	void Start();
-	void StartIfAsync();
+	USDUTILITIES_API void Start();
+	USDUTILITIES_API void StartIfAsync();
 
 	bool IsStarted() const
 	{
 		return Result.IsSet() && Result->IsValid();
 	}
 
-	bool DoWork();
+	USDUTILITIES_API bool DoWork();
 
 	bool IsDone() const
 	{
@@ -465,15 +465,15 @@ struct USDUTILITIES_API FSchemaTranslatorTask
 	}
 };
 
-class USDUTILITIES_API FUsdSchemaTranslatorTaskChain
+class FUsdSchemaTranslatorTaskChain
 {
 public:
 	virtual ~FUsdSchemaTranslatorTaskChain() = default;
 
-	FUsdSchemaTranslatorTaskChain& Do(ESchemaTranslationLaunchPolicy Policy, TFunction<bool()> Callable);
-	FUsdSchemaTranslatorTaskChain& Then(ESchemaTranslationLaunchPolicy Policy, TFunction<bool()> Callable);
+	USDUTILITIES_API FUsdSchemaTranslatorTaskChain& Do(ESchemaTranslationLaunchPolicy Policy, TFunction<bool()> Callable);
+	USDUTILITIES_API FUsdSchemaTranslatorTaskChain& Then(ESchemaTranslationLaunchPolicy Policy, TFunction<bool()> Callable);
 
-	ESchemaTranslationStatus Execute(bool bExclusiveSyncTasks = false);
+	USDUTILITIES_API ESchemaTranslationStatus Execute(bool bExclusiveSyncTasks = false);
 
 private:
 	TSharedPtr<FSchemaTranslatorTask> CurrentTask;
diff --git a/Engine/Plugins/Runtime/USDCore/Source/USDUtilities/Public/USDErrorUtils.h b/Engine/Plugins/Runtime/USDCore/Source/USDUtilities/Public/USDErrorUtils.h
--- a/Engine/Plugins/Runtime/USDCore/Source/USDUtilities/Public/USDErrorUtils.h
+++ b/Engine/Plugins/Runtime/USDCore/Source/USDUtilities/Public/USDErrorUtils.h
@@ -6,6 +6,8 @@
 #include "Logging/TokenizedMessage.h"
 #include "Misc/SourceLocation.h"
 
+#define UE_API USDUTILITIES_API
+
 namespace UsdUtils
 {
 	/**
@@ -31,7 +33,7 @@ namespace UsdUtils
 	USDUTILITIES_API bool ShowErrorsAndStopMonitoring(const FText& ToastMessage = FText());
 }
 
-class USDUTILITIES_API FUsdLogManager
+class FUsdLogManager
 {
 public:
 	// Copied from AssertionMacros.h as it's private there
@@ -51,31 +53,31 @@ public:
 
 	// Log a message that shows only on the output log.
 	// Prefer using the USD_LOG_INFO, USD_LOG_WARNING, USD_LOG_ERROR, USD_LOG_USERINFO, USD_LOG_USERWARNING, USD_LOG_USERERROR macros.
-	static void Log(EMessageSeverity::Type Severity, const FString& Message, uint32 MessageID);
+	static UE_API void Log(EMessageSeverity::Type Severity, const FString& Message, uint32 MessageID);
 
 	// Log a message that shows on the output log and the message log, if we're inside of a FScopedUsdMessageLog scope.
 	// Prefer using the USD_LOG_INFO, USD_LOG_WARNING, USD_LOG_ERROR, USD_LOG_USERINFO, USD_LOG_USERWARNING, USD_LOG_USERERROR macros.
-	static void Log(EMessageSeverity::Type Severity, const FText& Message, uint32 MessageID);
+	static UE_API void Log(EMessageSeverity::Type Severity, const FText& Message, uint32 MessageID);
 
 	// Returns whether we have currently accumulated any message with severity warning or higher.
 	// If we're not in the scope of any FScopedUsdMessageLog, returns false.
-	static bool HasAccumulatedErrors();
+	static UE_API bool HasAccumulatedErrors();
 
-	static void RegisterDiagnosticDelegate();
+	static UE_API void RegisterDiagnosticDelegate();
 
-	static void UnregisterDiagnosticDelegate();
+	static UE_API void UnregisterDiagnosticDelegate();
 
 	UE_DEPRECATED(5.6, "Use Log, or ideally the USD_LOG_INFO, USD_LOG_WARNING, etc. macros")
-	static void LogMessage(EMessageSeverity::Type Severity, const FText& Message);
+	static UE_API void LogMessage(EMessageSeverity::Type Severity, const FText& Message);
 
 	UE_DEPRECATED(5.6, "Use Log, or ideally the USD_LOG_INFO, USD_LOG_WARNING, etc. macros")
-	static void LogMessage(const TSharedRef<FTokenizedMessage>& Message);
+	static UE_API void LogMessage(const TSharedRef<FTokenizedMessage>& Message);
 
 	UE_DEPRECATED(5.6, "Used FScopedUsdMessageLog instead.")
-	static void EnableMessageLog();
+	static UE_API void EnableMessageLog();
 
 	UE_DEPRECATED(5.6, "Used FScopedUsdMessageLog instead.")
-	static void DisableMessageLog();
+	static UE_API void DisableMessageLog();
 };
 
 /**
@@ -106,14 +108,16 @@ public:
  * If the USD.UseMessageLog cvar is enabled (default), it will add the user-facing messages (and the USD SDK and error mark messages) to the Message Log,
  * (whether deduplication is enabled or disabled), also displaying a toast letting the user know that those messages can be viewed there.
  */
-class USDUTILITIES_API FScopedUsdMessageLog
+class FScopedUsdMessageLog
 {
 public:
-	FScopedUsdMessageLog();
-	~FScopedUsdMessageLog();
+	UE_API FScopedUsdMessageLog();
+	UE_API ~FScopedUsdMessageLog();
 
 	FScopedUsdMessageLog(const FScopedUsdMessageLog&) = delete;
 	FScopedUsdMessageLog& operator=(const FScopedUsdMessageLog&) = delete;
 	FScopedUsdMessageLog(FScopedUsdMessageLog&&) = delete;
 	FScopedUsdMessageLog& operator=(FScopedUsdMessageLog&&) = delete;
 };
+
+#undef UE_API
diff --git a/Engine/Plugins/Runtime/USDCore/Source/USDUtilities/Public/USDExportUtils.h b/Engine/Plugins/Runtime/USDCore/Source/USDUtilities/Public/USDExportUtils.h
--- a/Engine/Plugins/Runtime/USDCore/Source/USDUtilities/Public/USDExportUtils.h
+++ b/Engine/Plugins/Runtime/USDCore/Source/USDUtilities/Public/USDExportUtils.h
@@ -4,6 +4,8 @@
 
 #include "CoreMinimal.h"
 
+#define UE_API USDUTILITIES_API
+
 namespace UsdUnreal::ExportUtils
 {
 	/**
@@ -25,11 +27,11 @@ namespace UsdUnreal::ExportUtils
 	USDUTILITIES_API void EndUniquePathScope();
 
 	/** Utility to call BeginUniquePathScope() on construction and EndUniquePathScope() on destruction */
-	class USDUTILITIES_API FUniquePathScope final
+	class FUniquePathScope final
 	{
 	public:
-		FUniquePathScope();
-		~FUniquePathScope();
+		UE_API FUniquePathScope();
+		UE_API ~FUniquePathScope();
 
 		FUniquePathScope(const FUniquePathScope&) = delete;
 		FUniquePathScope(FUniquePathScope&&) = delete;
@@ -51,3 +53,5 @@ namespace UsdUnreal::ExportUtils
 	 */
 	USDUTILITIES_API void SanitizeFilePath(FString& Path);
 }
+
+#undef UE_API
diff --git a/Engine/Plugins/Runtime/USDCore/Source/USDUtilities/Public/USDGeomMeshConversion.h b/Engine/Plugins/Runtime/USDCore/Source/USDUtilities/Public/USDGeomMeshConversion.h
--- a/Engine/Plugins/Runtime/USDCore/Source/USDUtilities/Public/USDGeomMeshConversion.h
+++ b/Engine/Plugins/Runtime/USDCore/Source/USDUtilities/Public/USDGeomMeshConversion.h
@@ -54,7 +54,7 @@ namespace UE
 namespace UsdToUnreal
 {
 	/** Common options to mesh conversion functions */
-	struct USDUTILITIES_API FUsdMeshConversionOptions
+	struct FUsdMeshConversionOptions
 	{
 		/** This enum must match EInterchangeUsdPrimvar in InterchangeUsdDefinitions.h*/
 		enum EImportPrimvar
@@ -74,10 +74,10 @@ namespace UsdToUnreal
 		PRAGMA_DISABLE_DEPRECATION_WARNINGS
 		FUsdMeshConversionOptions(const FUsdMeshConversionOptions& Other) = default;
 		FUsdMeshConversionOptions& operator=(const FUsdMeshConversionOptions& Other) = default;
-		PRAGMA_ENABLE_DEPRECATION_WARNINGS
+		USDUTILITIES_API PRAGMA_ENABLE_DEPRECATION_WARNINGS
 
 		const static FUsdMeshConversionOptions DefaultOptions;
-		FUsdMeshConversionOptions();
+		USDUTILITIES_API FUsdMeshConversionOptions();
 
 		// List of the primvars that should be stored on the MeshDescription
 		TSet<FString> PrimvarNames;
diff --git a/Engine/Plugins/Runtime/USDCore/Source/USDUtilities/Public/USDMaterialXShaderGraph.h b/Engine/Plugins/Runtime/USDCore/Source/USDUtilities/Public/USDMaterialXShaderGraph.h
--- a/Engine/Plugins/Runtime/USDCore/Source/USDUtilities/Public/USDMaterialXShaderGraph.h
+++ b/Engine/Plugins/Runtime/USDCore/Source/USDUtilities/Public/USDMaterialXShaderGraph.h
@@ -18,7 +18,7 @@
  * Create a MaterialX shader graph from a UsdPrim, all <geompropvalue> nodes are either converted to an <image> or a proper <geomcolor>
  * Just get the Document in order to search for the surface material, by just looking for the name of the Prim it will return the correct surface material.
  */
-class USDUTILITIES_API FUsdMaterialXShaderGraph
+class FUsdMaterialXShaderGraph
 {
 public:
 
@@ -34,25 +34,25 @@ public:
 	 * @param UsdShadeMaterialPrim - The prim used to generate the shader graph, should be a valid pxr::UsdShadeMaterial
 	 * @param RenderContext - Which render context output to read from the UsdShadeMaterial
 	 */
-	FUsdMaterialXShaderGraph(const pxr::UsdPrim& MaterialPrim, const TCHAR* RenderContext);
+	USDUTILITIES_API FUsdMaterialXShaderGraph(const pxr::UsdPrim& MaterialPrim, const TCHAR* RenderContext);
 
-	MaterialX::DocumentPtr GetDocument() const;
+	USDUTILITIES_API MaterialX::DocumentPtr GetDocument() const;
 
 	/** Return the names of the <geompropvalue> nodes which have converted to <image> node*/
-	const TArray<FGeomProp>& GetGeomPropValueNames() const;
+	USDUTILITIES_API const TArray<FGeomProp>& GetGeomPropValueNames() const;
 
 private:
 
-	void CreateSurfaceMaterial(const pxr::UsdShadeMaterial& UsdShadeMaterial, const TCHAR* RenderContext);
+	USDUTILITIES_API void CreateSurfaceMaterial(const pxr::UsdShadeMaterial& UsdShadeMaterial, const TCHAR* RenderContext);
 
-	bool ComputeShaderGraph(MaterialX::NodePtr Node, const pxr::UsdShadeShader& Shader);
+	USDUTILITIES_API bool ComputeShaderGraph(MaterialX::NodePtr Node, const pxr::UsdShadeShader& Shader);
 
-	bool IsGeomColor(const pxr::UsdShadeShader& GeomPropValueShader);
+	USDUTILITIES_API bool IsGeomColor(const pxr::UsdShadeShader& GeomPropValueShader);
 
-	void SetInputValue(MaterialX::InputPtr Input, const pxr::UsdShadeInput& UsdInput);
+	USDUTILITIES_API void SetInputValue(MaterialX::InputPtr Input, const pxr::UsdShadeInput& UsdInput);
 
 	/** Set the interfaceName of the input if any */ 
-	bool SetInterfaceName(MaterialX::InputPtr Input, const pxr::UsdShadeConnectionSourceInfo& Source);
+	USDUTILITIES_API bool SetInterfaceName(MaterialX::InputPtr Input, const pxr::UsdShadeConnectionSourceInfo& Source);
 
 	/**
 	 * Get the MaterialX node from a UsdShadeConnectionSourceInfo
@@ -60,11 +60,11 @@ private:
 	 * @param bIsNodeAlreadyExist - return true if the node was already present in the shader graph, false otherwise
 	 * @return The corresponding MaterialX node in the shader graph, if it doesn't exist it creates one
 	 */
-	MaterialX::NodePtr GetNode(const pxr::UsdShadeConnectionSourceInfo& Source, bool& bIsNodeAlreadyExist);
+	USDUTILITIES_API MaterialX::NodePtr GetNode(const pxr::UsdShadeConnectionSourceInfo& Source, bool& bIsNodeAlreadyExist);
 
-	MaterialX::InputPtr GetInput(MaterialX::NodePtr Node, const pxr::UsdShadeInput& UsdInput);
+	USDUTILITIES_API MaterialX::InputPtr GetInput(MaterialX::NodePtr Node, const pxr::UsdShadeInput& UsdInput);
 
-	MaterialX::NodeDefPtr GetNodeDef(const pxr::UsdShadeConnectionSourceInfo& Source, bool& bIsGeomPropValue);
+	USDUTILITIES_API MaterialX::NodeDefPtr GetNodeDef(const pxr::UsdShadeConnectionSourceInfo& Source, bool& bIsGeomPropValue);
 
 	/**
 	 * Given an input type, retrieve the MaterialX associated type
@@ -72,9 +72,9 @@ private:
 	 * @param InOutInputType - The type of the input, if it's valid it will now be the MaterialX type
 	 * @return true if there is a valid match between USD and MaterialX or if it's already a MaterialX type, false otherwise
 	 */
-	bool GetMatchingInputType(const std::string& InputName, std::string& InOutInputType);
+	USDUTILITIES_API bool GetMatchingInputType(const std::string& InputName, std::string& InOutInputType);
 
-	bool IsTangentSpaceInput(const pxr::UsdShadeInput& UsdInput) const;
+	USDUTILITIES_API bool IsTangentSpaceInput(const pxr::UsdShadeInput& UsdInput) const;
 
 private:
 
@@ -111,4 +111,4 @@ private:
 		               pxr::GfVec2f,
 		   MtlxType>>>;
 #endif // USE_USD_SDK && ENABLE_USD_MATERIALX
-};
\ No newline at end of file
+};
diff --git a/Engine/Plugins/Runtime/USDCore/Source/USDUtilities/Public/USDPrimConversion.h b/Engine/Plugins/Runtime/USDCore/Source/USDUtilities/Public/USDPrimConversion.h
--- a/Engine/Plugins/Runtime/USDCore/Source/USDUtilities/Public/USDPrimConversion.h
+++ b/Engine/Plugins/Runtime/USDCore/Source/USDUtilities/Public/USDPrimConversion.h
@@ -186,7 +186,7 @@ namespace UsdToUnreal
 	 * This is used so that the reading process can be decoupled from the baking process (ConvertXTimeSamples functions) and reused for different
 	 * attributes c.f. CreatePropertyTrackReader
 	 */
-	struct USDUTILITIES_API FPropertyTrackReader
+	struct FPropertyTrackReader
 	{
 		TFunction<float(double)> FloatReader;
 		TFunction<bool(double)> BoolReader;
@@ -384,7 +384,7 @@ namespace UnrealToUsd
 	const static inline int32 NumBakingTypes = 6;
 
 	// Contains a lambda function responsible for baking a USceneComponent on the level into a prim on an FUsdStage
-	struct USDUTILITIES_API FComponentBaker
+	struct FComponentBaker
 	{
 		EBakingType BakerType;
 		FString ComponentPath;	  // Used for sorting, to ensure we evaluate parents and skeletal stuff before children and attached stuff
@@ -425,7 +425,7 @@ namespace UnrealToUsd
 	 * This is used so that the writing process can be decoupled from the baking process (ConvertXTrack functions) and reused for different attributes
 	 * c.f. CreatePropertyTrackWriter
 	 */
-	struct USDUTILITIES_API FPropertyTrackWriter
+	struct FPropertyTrackWriter
 	{
 		TFunction<void(float, double)> FloatWriter;
 		TFunction<void(bool, double)> BoolWriter;
diff --git a/Engine/Plugins/Runtime/USDCore/Source/USDUtilities/Public/USDShadeConversion.h b/Engine/Plugins/Runtime/USDCore/Source/USDUtilities/Public/USDShadeConversion.h
--- a/Engine/Plugins/Runtime/USDCore/Source/USDUtilities/Public/USDShadeConversion.h
+++ b/Engine/Plugins/Runtime/USDCore/Source/USDUtilities/Public/USDShadeConversion.h
@@ -51,7 +51,7 @@ namespace UE
 
 namespace UsdToUnreal
 {
-	struct USDUTILITIES_API FTextureParameterValue
+	struct FTextureParameterValue
 	{
 		UTexture* Texture = nullptr;	// Only used for the ConvertMaterial overloads that receive TexturesCache
 
@@ -73,10 +73,10 @@ namespace UsdToUnreal
 	public:
 		// Returns whether the texture should be parsed as sRGB or not, given the actually
 		// authored "bSRGB" member and the fallback opinion provided by the texture group
-		bool GetSRGBValue() const;
+		USDUTILITIES_API bool GetSRGBValue() const;
 	};
 
-	struct USDUTILITIES_API FPrimvarReaderParameterValue
+	struct FPrimvarReaderParameterValue
 	{
 		FString PrimvarName;
 		FVector FallbackValue;
@@ -84,7 +84,7 @@ namespace UsdToUnreal
 
 	using FParameterValue = TVariant<float, FVector, FTextureParameterValue, FPrimvarReaderParameterValue, bool>;
 
-	struct USDUTILITIES_API FUsdPreviewSurfaceMaterialData
+	struct FUsdPreviewSurfaceMaterialData
 	{
 		TMap<FString, FParameterValue> Parameters;
 
diff --git a/Engine/Plugins/Runtime/USDCore/Source/USDUtilities/Public/USDSkeletalDataConversion.h b/Engine/Plugins/Runtime/USDCore/Source/USDUtilities/Public/USDSkeletalDataConversion.h
--- a/Engine/Plugins/Runtime/USDCore/Source/USDUtilities/Public/USDSkeletalDataConversion.h
+++ b/Engine/Plugins/Runtime/USDCore/Source/USDUtilities/Public/USDSkeletalDataConversion.h
@@ -63,7 +63,7 @@ namespace UE
 
 namespace UsdUtils
 {
-	struct USDUTILITIES_API FUsdBlendShapeInbetween
+	struct FUsdBlendShapeInbetween
 	{
 		// Name of the FUsdBlendShape/UMorphTarget that holds the morph data for this inbetween
 		FString Name;
@@ -78,7 +78,7 @@ namespace UsdUtils
 		}
 	};
 
-	struct USDUTILITIES_API FUsdBlendShape
+	struct FUsdBlendShape
 	{
 		FString Name;
 
diff --git a/Engine/Plugins/Runtime/USDCore/Source/USDUtilities/Public/USDTypesConversion.h b/Engine/Plugins/Runtime/USDCore/Source/USDUtilities/Public/USDTypesConversion.h
--- a/Engine/Plugins/Runtime/USDCore/Source/USDUtilities/Public/USDTypesConversion.h
+++ b/Engine/Plugins/Runtime/USDCore/Source/USDUtilities/Public/USDTypesConversion.h
@@ -50,12 +50,12 @@ PXR_NAMESPACE_CLOSE_SCOPE
 template<typename UsdObjectType>
 class TUsdStore;
 
-struct USDUTILITIES_API FUsdStageInfo
+struct FUsdStageInfo
 {
 	EUsdUpAxis UpAxis = EUsdUpAxis::ZAxis;
 	double MetersPerUnit = 0.01;
 
-	explicit FUsdStageInfo(const pxr::UsdStageRefPtr& Stage);
+	USDUTILITIES_API explicit FUsdStageInfo(const pxr::UsdStageRefPtr& Stage);
 };
 
 namespace UsdToUnreal
diff --git a/Engine/Plugins/Runtime/USDCore/Source/UnrealUSDWrapper/Public/USDListener.h b/Engine/Plugins/Runtime/USDCore/Source/UnrealUSDWrapper/Public/USDListener.h
--- a/Engine/Plugins/Runtime/USDCore/Source/UnrealUSDWrapper/Public/USDListener.h
+++ b/Engine/Plugins/Runtime/USDCore/Source/UnrealUSDWrapper/Public/USDListener.h
@@ -9,6 +9,8 @@
 #include "UsdWrappers/ForwardDeclarations.h"
 #include "UsdWrappers/VtValue.h"
 
+#define UE_API UNREALUSDWRAPPER_API
+
 class FUsdListenerImpl;
 
 namespace UE
@@ -96,59 +98,59 @@ namespace UsdUtils
 /**
  * Registers to Usd Notices and emits events when the Usd Stage has changed
  */
-class UNREALUSDWRAPPER_API FUsdListener
+class FUsdListener
 {
 public:
-	FUsdListener();
-	FUsdListener(const UE::FUsdStage& Stage);
+	UE_API FUsdListener();
+	UE_API FUsdListener(const UE::FUsdStage& Stage);
 
 	FUsdListener(const FUsdListener& Other) = delete;
 	FUsdListener(FUsdListener&& Other) = delete;
 
-	virtual ~FUsdListener();
+	UE_API virtual ~FUsdListener();
 
 	FUsdListener& operator=(const FUsdListener& Other) = delete;
 	FUsdListener& operator=(FUsdListener&& Other) = delete;
 
-	void Register(const UE::FUsdStage& Stage);
+	UE_API void Register(const UE::FUsdStage& Stage);
 
 	// Increment/decrement the block counter
-	void Block();
-	void Unblock();
-	bool IsBlocked() const;
+	UE_API void Block();
+	UE_API void Unblock();
+	UE_API bool IsBlocked() const;
 
 	// Stage-specific events
 	DECLARE_TS_MULTICAST_DELEGATE(FOnStageEditTargetChanged);
-	FOnStageEditTargetChanged& GetOnStageEditTargetChanged();
+	UE_API FOnStageEditTargetChanged& GetOnStageEditTargetChanged();
 
 	DECLARE_TS_MULTICAST_DELEGATE_OneParam(FOnLayersChanged, const TArray<FString>&);
 	UE_DEPRECATED(5.3, "Use GetOnSdfLayersChanged")
-	PRAGMA_DISABLE_DEPRECATION_WARNINGS
+	UE_API PRAGMA_DISABLE_DEPRECATION_WARNINGS
 	FOnLayersChanged& GetOnLayersChanged();
 	PRAGMA_ENABLE_DEPRECATION_WARNINGS
 
 	DECLARE_TS_MULTICAST_DELEGATE_OneParam(FOnSdfLayersChanged, const UsdUtils::FLayerToSdfChangeList&);
-	FOnSdfLayersChanged& GetOnSdfLayersChanged();
+	UE_API FOnSdfLayersChanged& GetOnSdfLayersChanged();
 
 	DECLARE_TS_MULTICAST_DELEGATE(FOnSdfLayerDirtinessChanged);
-	FOnSdfLayerDirtinessChanged& GetOnSdfLayerDirtinessChanged();
+	UE_API FOnSdfLayerDirtinessChanged& GetOnSdfLayerDirtinessChanged();
 
 	DECLARE_TS_MULTICAST_DELEGATE_TwoParams(
 		FOnObjectsChanged,
 		const UsdUtils::FObjectChangesByPath& /* InfoChanges */,
 		const UsdUtils::FObjectChangesByPath& /* ResyncChanges */
 	);
-	FOnObjectsChanged& GetOnObjectsChanged();
+	UE_API FOnObjectsChanged& GetOnObjectsChanged();
 
 private:
 	TUniquePtr<FUsdListenerImpl> Impl;
 };
 
-class UNREALUSDWRAPPER_API FScopedBlockNotices final
+class FScopedBlockNotices final
 {
 public:
-	explicit FScopedBlockNotices(FUsdListener& InListener);
-	~FScopedBlockNotices();
+	UE_API explicit FScopedBlockNotices(FUsdListener& InListener);
+	UE_API ~FScopedBlockNotices();
 
 	FScopedBlockNotices() = delete;
 	FScopedBlockNotices(const FScopedBlockNotices&) = delete;
@@ -159,3 +161,5 @@ public:
 private:
 	FUsdListener& Listener;
 };
+
+#undef UE_API
diff --git a/Engine/Plugins/Runtime/USDCore/Source/UnrealUSDWrapper/Public/USDMemory.h b/Engine/Plugins/Runtime/USDCore/Source/UnrealUSDWrapper/Public/USDMemory.h
--- a/Engine/Plugins/Runtime/USDCore/Source/UnrealUSDWrapper/Public/USDMemory.h
+++ b/Engine/Plugins/Runtime/USDCore/Source/UnrealUSDWrapper/Public/USDMemory.h
@@ -12,6 +12,8 @@
 #include "Templates/SharedPointer.h"
 #include "Templates/UnrealTemplate.h"
 
+#define UE_API UNREALUSDWRAPPER_API
+
 LLM_DECLARE_TAG_API(Usd, UNREALUSDWRAPPER_API);
 
 /**
@@ -59,31 +61,31 @@ enum class EUsdActiveAllocator
 	System
 };
 
-class UNREALUSDWRAPPER_API FUsdMemoryManager
+class FUsdMemoryManager
 {
 public:
-	static void Initialize();
-	static void Shutdown();
+	static UE_API void Initialize();
+	static UE_API void Shutdown();
 
 	/** Pushes Allocator on the stack of active allocators */
-	static void ActivateAllocator(EUsdActiveAllocator Allocator);
+	static UE_API void ActivateAllocator(EUsdActiveAllocator Allocator);
 
 	/** Pops Allocator from the stack of active allocators */
-	static bool DeactivateAllocator(EUsdActiveAllocator Allocator);
+	static UE_API bool DeactivateAllocator(EUsdActiveAllocator Allocator);
 
 	/** Redirects the call to malloc to the currently active allocator. */
-	static void* Malloc(SIZE_T Count);
+	static UE_API void* Malloc(SIZE_T Count);
 
 	/** Redirects the call to free to the currently active allocator. */
-	static void Free(void* Original);
+	static UE_API void Free(void* Original);
 
 private:
-	static FActiveAllocatorsStack& GetActiveAllocatorsStackForThread();
+	static UE_API FActiveAllocatorsStack& GetActiveAllocatorsStackForThread();
 
 	/** Returns if the current active allocator is EActiveAllocator::System */
-	static bool IsUsingSystemMalloc();
+	static UE_API bool IsUsingSystemMalloc();
 
-	static TOptional<FTlsSlot> ActiveAllocatorsStackTLS;
+	static UE_API TOptional<FTlsSlot> ActiveAllocatorsStackTLS;
 
 	struct FThreadSafeSet
 	{
@@ -113,9 +115,9 @@ private:
 		FSubSet SubSets[BucketCount];
 	};
 
-	static FThreadSafeSet SystemAllocedPtrs;
+	static UE_API FThreadSafeSet SystemAllocedPtrs;
 
-	static FCriticalSection CriticalSection;
+	static UE_API FCriticalSection CriticalSection;
 };
 
 /**
@@ -328,3 +330,5 @@ TSharedRef<ObjectType> MakeSharedUnreal(ArgTypes&&... Args)
 	#define IMPLEMENT_MODULE_USD( ModuleImplClass, ModuleName ) IMPLEMENT_MODULE( ModuleImplClass, ModuleName )
 #endif
 // clang-format on
+
+#undef UE_API
diff --git a/Engine/Plugins/Runtime/USDCore/Source/UnrealUSDWrapper/Public/UnrealUSDWrapper.h b/Engine/Plugins/Runtime/USDCore/Source/UnrealUSDWrapper/Public/UnrealUSDWrapper.h
--- a/Engine/Plugins/Runtime/USDCore/Source/UnrealUSDWrapper/Public/UnrealUSDWrapper.h
+++ b/Engine/Plugins/Runtime/USDCore/Source/UnrealUSDWrapper/Public/UnrealUSDWrapper.h
@@ -593,9 +593,9 @@ namespace UnrealIdentifiers
 
 }	 // namespace UnrealIdentifiers
 
-struct UNREALUSDWRAPPER_API FUsdDelegates
+struct FUsdDelegates
 {
 	DECLARE_MULTICAST_DELEGATE_OneParam(FUsdImportDelegate, FString /* FilePath */);
-	static FUsdImportDelegate OnPreUsdImport;
-	static FUsdImportDelegate OnPostUsdImport;
+	static UNREALUSDWRAPPER_API FUsdImportDelegate OnPreUsdImport;
+	static UNREALUSDWRAPPER_API FUsdImportDelegate OnPostUsdImport;
 };
diff --git a/Engine/Plugins/Runtime/USDCore/Source/UnrealUSDWrapper/Public/UsdWrappers/UsdEditContext.h b/Engine/Plugins/Runtime/USDCore/Source/UnrealUSDWrapper/Public/UsdWrappers/UsdEditContext.h
--- a/Engine/Plugins/Runtime/USDCore/Source/UnrealUSDWrapper/Public/UsdWrappers/UsdEditContext.h
+++ b/Engine/Plugins/Runtime/USDCore/Source/UnrealUSDWrapper/Public/UsdWrappers/UsdEditContext.h
@@ -6,6 +6,8 @@
 
 #include "Templates/UniquePtr.h"
 
+#define UE_API UNREALUSDWRAPPER_API
+
 namespace UE
 {
 	namespace Internal
@@ -16,17 +18,19 @@ namespace UE
 	/**
 	 * Minimal pxr::UsdEditContext wrapper for Unreal that can be used from no-rtti modules.
 	 */
-	class UNREALUSDWRAPPER_API FUsdEditContext final
+	class FUsdEditContext final
 	{
 		FUsdEditContext(const FUsdEditContext& Other) = delete;
 		FUsdEditContext& operator=(const FUsdEditContext& Other) = delete;
 
 	public:
-		explicit FUsdEditContext(const FUsdStageWeak& Stage);
-		FUsdEditContext(const FUsdStageWeak& Stage, const FSdfLayer& EditTarget);
-		~FUsdEditContext();
+		UE_API explicit FUsdEditContext(const FUsdStageWeak& Stage);
+		UE_API FUsdEditContext(const FUsdStageWeak& Stage, const FSdfLayer& EditTarget);
+		UE_API ~FUsdEditContext();
 
 	private:
 		TUniquePtr<Internal::FUsdEditContextImpl> Impl;
 	};
 }
+
+#undef UE_API
diff --git a/Engine/Plugins/Runtime/USDCore/Source/UnrealUSDWrapper/Public/UsdWrappers/UsdGeomBBoxCache.h b/Engine/Plugins/Runtime/USDCore/Source/UnrealUSDWrapper/Public/UsdWrappers/UsdGeomBBoxCache.h
--- a/Engine/Plugins/Runtime/USDCore/Source/UnrealUSDWrapper/Public/UsdWrappers/UsdGeomBBoxCache.h
+++ b/Engine/Plugins/Runtime/USDCore/Source/UnrealUSDWrapper/Public/UsdWrappers/UsdGeomBBoxCache.h
@@ -35,40 +35,40 @@ namespace UE
 	 * and the conversion operators that return the underlying type will check that the mutex is
 	 * locked before returning (the caller is expected to lock the mutex before using the operators).
 	 */
-	class UNREALUSDWRAPPER_API FUsdGeomBBoxCache
+	class FUsdGeomBBoxCache
 	{
 	public:
-		FUsdGeomBBoxCache(double Time, EUsdPurpose IncludedPurposeFlags, bool bUseExtentsHint = false, bool bIgnoreVisibility = false);
-		FUsdGeomBBoxCache(const FUsdGeomBBoxCache& Other);
-		FUsdGeomBBoxCache(FUsdGeomBBoxCache&& Other);
+		UNREALUSDWRAPPER_API FUsdGeomBBoxCache(double Time, EUsdPurpose IncludedPurposeFlags, bool bUseExtentsHint = false, bool bIgnoreVisibility = false);
+		UNREALUSDWRAPPER_API FUsdGeomBBoxCache(const FUsdGeomBBoxCache& Other);
+		UNREALUSDWRAPPER_API FUsdGeomBBoxCache(FUsdGeomBBoxCache&& Other);
 
-		FUsdGeomBBoxCache& operator=(const FUsdGeomBBoxCache& Other);
-		FUsdGeomBBoxCache& operator=(FUsdGeomBBoxCache&& Other);
+		UNREALUSDWRAPPER_API FUsdGeomBBoxCache& operator=(const FUsdGeomBBoxCache& Other);
+		UNREALUSDWRAPPER_API FUsdGeomBBoxCache& operator=(FUsdGeomBBoxCache&& Other);
 
-		~FUsdGeomBBoxCache();
+		UNREALUSDWRAPPER_API ~FUsdGeomBBoxCache();
 
 		// Auto conversion from/to pxr::UsdGeomBBoxCache
 	public:
 #if USE_USD_SDK
-		explicit FUsdGeomBBoxCache(const pxr::UsdGeomBBoxCache& InUsdGeomBBoxCache);
-		explicit FUsdGeomBBoxCache(pxr::UsdGeomBBoxCache&& InUsdGeomBBoxCache);
+		UNREALUSDWRAPPER_API explicit FUsdGeomBBoxCache(const pxr::UsdGeomBBoxCache& InUsdGeomBBoxCache);
+		UNREALUSDWRAPPER_API explicit FUsdGeomBBoxCache(pxr::UsdGeomBBoxCache&& InUsdGeomBBoxCache);
 
-		operator pxr::UsdGeomBBoxCache&();
-		operator const pxr::UsdGeomBBoxCache&() const;
+		UNREALUSDWRAPPER_API operator pxr::UsdGeomBBoxCache&();
+		UNREALUSDWRAPPER_API operator const pxr::UsdGeomBBoxCache&() const;
 #endif	  // #if USE_USD_SDK
 
 		  // Wrapped pxr::UsdGeomBBoxCache functions, refer to the USD SDK documentation
 	public:
-		void Clear();
+		UNREALUSDWRAPPER_API void Clear();
 
-		void SetIncludedPurposes(EUsdPurpose IncludedPurposeFlags);
-		EUsdPurpose GetIncludedPurposes() const;
+		UNREALUSDWRAPPER_API void SetIncludedPurposes(EUsdPurpose IncludedPurposeFlags);
+		UNREALUSDWRAPPER_API EUsdPurpose GetIncludedPurposes() const;
 
-		bool GetUseExtentsHint() const;
-		bool GetIgnoreVisibility() const;
+		UNREALUSDWRAPPER_API bool GetUseExtentsHint() const;
+		UNREALUSDWRAPPER_API bool GetIgnoreVisibility() const;
 
-		void SetTime(double UsdTimeCode);
-		double GetTime() const;
+		UNREALUSDWRAPPER_API void SetTime(double UsdTimeCode);
+		UNREALUSDWRAPPER_API double GetTime() const;
 
 	public:
 		// The underlying pxr::UsdGeomBBoxCache object is not thread-safe, so we have this.
diff --git a/Engine/Plugins/Runtime/WaveTable/Source/WaveTableEditor/Public/WaveTableBankEditor.h b/Engine/Plugins/Runtime/WaveTable/Source/WaveTableEditor/Public/WaveTableBankEditor.h
--- a/Engine/Plugins/Runtime/WaveTable/Source/WaveTableEditor/Public/WaveTableBankEditor.h
+++ b/Engine/Plugins/Runtime/WaveTable/Source/WaveTableEditor/Public/WaveTableBankEditor.h
@@ -8,6 +8,8 @@
 #include "Toolkits/AssetEditorToolkit.h"
 #include "WaveTableSettings.h"
 
+#define UE_API WAVETABLEEDITOR_API
+
 class FSpawnTabArgs;
 enum class EWaveTableCurveSource : uint8;
 namespace WaveTable::Editor { class FWaveTableCurveModel; }
@@ -25,30 +27,30 @@ namespace WaveTable
 {
 	namespace Editor
 	{
-		class WAVETABLEEDITOR_API FBankEditorBase : public FAssetEditorToolkit, public FNotifyHook, public FEditorUndoClient
+		class FBankEditorBase : public FAssetEditorToolkit, public FNotifyHook, public FEditorUndoClient
 		{
 		public:
-			FBankEditorBase();
+			UE_API FBankEditorBase();
 			virtual ~FBankEditorBase() = default;
 
-			void Init(const EToolkitMode::Type Mode, const TSharedPtr<IToolkitHost>& InitToolkitHost, UObject* InParentObject);
+			UE_API void Init(const EToolkitMode::Type Mode, const TSharedPtr<IToolkitHost>& InitToolkitHost, UObject* InParentObject);
 
 			/** FAssetEditorToolkit interface */
-			virtual FName GetToolkitFName() const override;
-			virtual FText GetBaseToolkitName() const override;
-			virtual FString GetWorldCentricTabPrefix() const override;
-			virtual FLinearColor GetWorldCentricTabColorScale() const override;
-			virtual void RegisterTabSpawners(const TSharedRef<FTabManager>& TabManager) override;
-			virtual void UnregisterTabSpawners(const TSharedRef<FTabManager>& TabManager) override;
+			UE_API virtual FName GetToolkitFName() const override;
+			UE_API virtual FText GetBaseToolkitName() const override;
+			UE_API virtual FString GetWorldCentricTabPrefix() const override;
+			UE_API virtual FLinearColor GetWorldCentricTabColorScale() const override;
+			UE_API virtual void RegisterTabSpawners(const TSharedRef<FTabManager>& TabManager) override;
+			UE_API virtual void UnregisterTabSpawners(const TSharedRef<FTabManager>& TabManager) override;
 
 			/** FNotifyHook interface */
-			virtual void NotifyPostChange(const FPropertyChangedEvent& PropertyChangedEvent, FProperty* PropertyThatChanged) override;
+			UE_API virtual void NotifyPostChange(const FPropertyChangedEvent& PropertyChangedEvent, FProperty* PropertyThatChanged) override;
 
 			// Regenerates curves set to "file" without refreshing whole stack view
-			void RegenerateFileCurves();
+			UE_API void RegenerateFileCurves();
 
 			/** Updates & redraws curves. */
-			void RefreshCurves();
+			UE_API void RefreshCurves();
 
 		protected:
 			struct FCurveData
@@ -65,8 +67,8 @@ namespace WaveTable
 				}
 			};
 
-			virtual void PostUndo(bool bSuccess) override;
-			virtual void PostRedo(bool bSuccess) override;
+			UE_API virtual void PostUndo(bool bSuccess) override;
+			UE_API virtual void PostRedo(bool bSuccess) override;
 
 			virtual bool GetIsPropertyEditorDisabled() const { return false; }
 
@@ -93,33 +95,33 @@ namespace WaveTable
 			// Returns the number of transforms associated with the given bank.
 			virtual int32 GetNumTransforms() const = 0;
 
-			void SetCurve(int32 InTransformIndex, FRichCurve& InRichCurve, EWaveTableCurveSource InSource);
+			UE_API void SetCurve(int32 InTransformIndex, FRichCurve& InRichCurve, EWaveTableCurveSource InSource);
 
 		private:
 			// Regenerates expression curve at the given index. If curve is a 'File' and source PCM data in Transform
 			// is too long, inline edits are not included in recalculation for better interact performance.
-			void GenerateExpressionCurve(FCurveData& OutCurveData, int32 InTransformIndex, bool bInIsUnset = false);
+			UE_API void GenerateExpressionCurve(FCurveData& OutCurveData, int32 InTransformIndex, bool bInIsUnset = false);
 
-			void InitCurves();
+			UE_API void InitCurves();
 
-			void ResetCurves();
+			UE_API void ResetCurves();
 
 			/**	Spawns the tab allowing for editing/viewing the output curve(s) */
-			TSharedRef<SDockTab> SpawnTab_OutputCurve(const FSpawnTabArgs& Args);
+			UE_API TSharedRef<SDockTab> SpawnTab_OutputCurve(const FSpawnTabArgs& Args);
 
 			/**	Spawns the tab allowing for editing/viewing the output curve(s) */
-			TSharedRef<SDockTab> SpawnTab_Properties(const FSpawnTabArgs& Args);
+			UE_API TSharedRef<SDockTab> SpawnTab_Properties(const FSpawnTabArgs& Args);
 
 			/** Get the orientation for the snap value controls. */
-			EOrientation GetSnapLabelOrientation() const;
+			UE_API EOrientation GetSnapLabelOrientation() const;
 
 			/** Trims keys out-of-bounds in provided curve */
-			static void TrimKeys(FRichCurve& OutCurve);
+			static UE_API void TrimKeys(FRichCurve& OutCurve);
 
 			/** Clears the expression curve at the given input index */
-			void ClearExpressionCurve(int32 InTransformIndex);
+			UE_API void ClearExpressionCurve(int32 InTransformIndex);
 
-			bool RequiresNewCurve(int32 InTransformIndex, const FRichCurve& InRichCurve) const;
+			UE_API bool RequiresNewCurve(int32 InTransformIndex, const FRichCurve& InRichCurve) const;
 
 			TSharedPtr<FUICommandList> ToolbarCurveTargetCommands;
 
@@ -132,26 +134,28 @@ namespace WaveTable
 			TSharedPtr<IDetailsView> PropertiesView;
 
 			/** Settings Editor App Identifier */
-			static const FName AppIdentifier;
-			static const FName CurveTabId;
-			static const FName PropertiesTabId;
+			static UE_API const FName AppIdentifier;
+			static UE_API const FName CurveTabId;
+			static UE_API const FName PropertiesTabId;
 		};
 
-		class WAVETABLEEDITOR_API FBankEditor : public FBankEditorBase
+		class FBankEditor : public FBankEditorBase
 		{
 		public:
 			FBankEditor() = default;
 			virtual ~FBankEditor() = default;
 
 		protected:
-			virtual EWaveTableResolution GetBankResolution() const override;
-			virtual EWaveTableSamplingMode GetBankSamplingMode() const override;
-			virtual int32 GetBankSampleRate() const override;
-			virtual bool GetBankIsBipolar() const override;
-			virtual int32 GetNumTransforms() const override;
-			virtual FWaveTableTransform* GetTransform(int32 InIndex) const override;
-
-			virtual TUniquePtr<FWaveTableCurveModel> ConstructCurveModel(FRichCurve& InRichCurve, UObject* InParentObject, EWaveTableCurveSource InSource) override;
+			UE_API virtual EWaveTableResolution GetBankResolution() const override;
+			UE_API virtual EWaveTableSamplingMode GetBankSamplingMode() const override;
+			UE_API virtual int32 GetBankSampleRate() const override;
+			UE_API virtual bool GetBankIsBipolar() const override;
+			UE_API virtual int32 GetNumTransforms() const override;
+			UE_API virtual FWaveTableTransform* GetTransform(int32 InIndex) const override;
+
+			UE_API virtual TUniquePtr<FWaveTableCurveModel> ConstructCurveModel(FRichCurve& InRichCurve, UObject* InParentObject, EWaveTableCurveSource InSource) override;
 		};
 	} // namespace Editor
 } // namespace WaveTable
+
+#undef UE_API
diff --git a/Engine/Plugins/Runtime/WaveTable/Source/WaveTableEditor/Public/WaveTableCurveEditorViewStacked.h b/Engine/Plugins/Runtime/WaveTable/Source/WaveTableEditor/Public/WaveTableCurveEditorViewStacked.h
--- a/Engine/Plugins/Runtime/WaveTable/Source/WaveTableEditor/Public/WaveTableCurveEditorViewStacked.h
+++ b/Engine/Plugins/Runtime/WaveTable/Source/WaveTableEditor/Public/WaveTableCurveEditorViewStacked.h
@@ -8,6 +8,8 @@
 
 #include "WaveTableCurveEditorViewStacked.generated.h"
 
+#define UE_API WAVETABLEEDITOR_API
+
 
 // Forward Declarations
 class UCurveFloat;
@@ -30,23 +32,23 @@ namespace WaveTable
 {
 	namespace Editor
 	{
-		class WAVETABLEEDITOR_API FWaveTableCurveModel : public FRichCurveEditorModelRaw
+		class FWaveTableCurveModel : public FRichCurveEditorModelRaw
 		{
 		public:
-			static ECurveEditorViewID WaveTableViewId;
+			static UE_API ECurveEditorViewID WaveTableViewId;
 
-			FWaveTableCurveModel(FRichCurve& InRichCurve, UObject* InOwner, EWaveTableCurveSource InSource);
+			UE_API FWaveTableCurveModel(FRichCurve& InRichCurve, UObject* InOwner, EWaveTableCurveSource InSource);
 
-			const FText& GetAxesDescriptor() const;
-			const UObject* GetParentObject() const;
-			EWaveTableCurveSource GetSource() const;
-			void Refresh(const FWaveTableTransform& InTransform, int32 InCurveIndex, bool bInIsBipolar, EWaveTableSamplingMode InSamplingMode = EWaveTableSamplingMode::FixedResolution);
+			UE_API const FText& GetAxesDescriptor() const;
+			UE_API const UObject* GetParentObject() const;
+			UE_API EWaveTableCurveSource GetSource() const;
+			UE_API void Refresh(const FWaveTableTransform& InTransform, int32 InCurveIndex, bool bInIsBipolar, EWaveTableSamplingMode InSamplingMode = EWaveTableSamplingMode::FixedResolution);
 
 			virtual ECurveEditorViewID GetViewId() const { return WaveTableViewId; }
-			virtual bool IsReadOnly() const override;
-			virtual FLinearColor GetColor() const override;
-			virtual void GetTimeRange(double& MinValue, double& MaxValue) const override;
-			virtual void GetValueRange(double& MinValue, double& MaxValue) const override;
+			UE_API virtual bool IsReadOnly() const override;
+			UE_API virtual FLinearColor GetColor() const override;
+			UE_API virtual void GetTimeRange(double& MinValue, double& MaxValue) const override;
+			UE_API virtual void GetValueRange(double& MinValue, double& MaxValue) const override;
 
 			float GetCurveDuration() const { return CurveDuration; }
 			int32 GetCurveIndex() const { return CurveIndex; }
@@ -56,10 +58,10 @@ namespace WaveTable
 			int32 GetNumSamples() const { return NumSamples; }
 			EWaveTableSamplingMode GetSamplingMode() const { return SamplingMode; }
 		protected:
-			virtual void RefreshCurveDescriptorText(const FWaveTableTransform& InTransform, FText& OutShortDisplayName, FText& OutInputAxisName, FText& OutOutputAxisName);
-			virtual FColor GetCurveColor() const;
-			virtual bool GetPropertyEditorDisabled() const;
-			virtual FText GetPropertyEditorDisabledText() const;
+			UE_API virtual void RefreshCurveDescriptorText(const FWaveTableTransform& InTransform, FText& OutShortDisplayName, FText& OutInputAxisName, FText& OutOutputAxisName);
+			UE_API virtual FColor GetCurveColor() const;
+			UE_API virtual bool GetPropertyEditorDisabled() const;
+			UE_API virtual FText GetPropertyEditorDisabledText() const;
 
 			TWeakObjectPtr<UObject> ParentObject;
 
@@ -80,13 +82,13 @@ namespace WaveTable
 			FText AxesDescriptor;
 		};
 
-		class WAVETABLEEDITOR_API SViewStacked : public SCurveEditorViewStacked
+		class SViewStacked : public SCurveEditorViewStacked
 		{
 		public:
-			void Construct(const FArguments& InArgs, TWeakPtr<FCurveEditor> InCurveEditor);
+			UE_API void Construct(const FArguments& InArgs, TWeakPtr<FCurveEditor> InCurveEditor);
 
 		protected:
-			virtual void PaintView(
+			UE_API virtual void PaintView(
 				const FPaintArgs& Args,
 				const FGeometry& AllottedGeometry,
 				const FSlateRect& MyCullingRect,
@@ -95,12 +97,12 @@ namespace WaveTable
 				const FWidgetStyle& InWidgetStyle,
 				bool bParentEnabled) const override;
 
-			virtual void DrawViewGrids(const FGeometry& AllottedGeometry, const FSlateRect& MyCullingRect, FSlateWindowElementList& OutDrawElements, int32 BaseLayerId, ESlateDrawEffect DrawEffects) const override;
-			virtual void DrawLabels(const FGeometry& AllottedGeometry, const FSlateRect& MyCullingRect, FSlateWindowElementList& OutDrawElements, int32 BaseLayerId, ESlateDrawEffect DrawEffects) const override;
+			UE_API virtual void DrawViewGrids(const FGeometry& AllottedGeometry, const FSlateRect& MyCullingRect, FSlateWindowElementList& OutDrawElements, int32 BaseLayerId, ESlateDrawEffect DrawEffects) const override;
+			UE_API virtual void DrawLabels(const FGeometry& AllottedGeometry, const FSlateRect& MyCullingRect, FSlateWindowElementList& OutDrawElements, int32 BaseLayerId, ESlateDrawEffect DrawEffects) const override;
 
-			virtual void FormatInputLabel(const FWaveTableCurveModel& EditorModel, const FNumberFormattingOptions& InLabelFormat, FText& InOutLabel) const;
+			UE_API virtual void FormatInputLabel(const FWaveTableCurveModel& EditorModel, const FNumberFormattingOptions& InLabelFormat, FText& InOutLabel) const;
 			virtual void FormatOutputLabel(const FWaveTableCurveModel& EditorModel, const FNumberFormattingOptions& InLabelFormat, FText& InOutLabel) const { }
-			virtual FText FormatToolTipTime(const FCurveModel& CurveModel, double EvaluatedTime) const override;
+			UE_API virtual FText FormatToolTipTime(const FCurveModel& CurveModel, double EvaluatedTime) const override;
 
 		private:
 			struct FGridDrawInfo
@@ -143,8 +145,10 @@ namespace WaveTable
 				double GetPixelTop() const;
 			};
 
-			void DrawViewGridLineX(FSlateWindowElementList& OutDrawElements, FGridDrawInfo& DrawInfo, ESlateDrawEffect DrawEffect, double OffsetAlpha, bool bIsMajor) const;
-			void DrawViewGridLineY(const float VerticalLine, FSlateWindowElementList& OutDrawElements, FGridDrawInfo &DrawInfo, ESlateDrawEffect DrawEffects, const FText* Label, bool bIsMajor) const;
+			UE_API void DrawViewGridLineX(FSlateWindowElementList& OutDrawElements, FGridDrawInfo& DrawInfo, ESlateDrawEffect DrawEffect, double OffsetAlpha, bool bIsMajor) const;
+			UE_API void DrawViewGridLineY(const float VerticalLine, FSlateWindowElementList& OutDrawElements, FGridDrawInfo &DrawInfo, ESlateDrawEffect DrawEffects, const FText* Label, bool bIsMajor) const;
 		};
 	} // namespace Editor
 } // namespace WaveTable
+
+#undef UE_API
diff --git a/Engine/Plugins/Runtime/WaveTable/Source/WaveTableEditor/Public/WaveTableTransformLayout.h b/Engine/Plugins/Runtime/WaveTable/Source/WaveTableEditor/Public/WaveTableTransformLayout.h
--- a/Engine/Plugins/Runtime/WaveTable/Source/WaveTableEditor/Public/WaveTableTransformLayout.h
+++ b/Engine/Plugins/Runtime/WaveTable/Source/WaveTableEditor/Public/WaveTableTransformLayout.h
@@ -3,6 +3,8 @@
 
 #include "IPropertyTypeCustomization.h"
 
+#define UE_API WAVETABLEEDITOR_API
+
 
 // Forward Declarations
 enum class EWaveTableCurve : uint8;
@@ -12,7 +14,7 @@ struct FWaveTableTransform;
 
 namespace WaveTable::Editor
 {
-	class WAVETABLEEDITOR_API FWaveTableDataLayoutCustomization : public IPropertyTypeCustomization
+	class FWaveTableDataLayoutCustomization : public IPropertyTypeCustomization
 	{
 	public:
 		static TSharedRef<IPropertyTypeCustomization> MakeInstance()
@@ -20,28 +22,28 @@ namespace WaveTable::Editor
 			return MakeShared<FWaveTableDataLayoutCustomization>();
 		}
 
-		virtual void CustomizeHeader(TSharedRef<IPropertyHandle> StructPropertyHandle, FDetailWidgetRow& HeaderRow, IPropertyTypeCustomizationUtils& CustomizationUtils) override;
-		virtual void CustomizeChildren(TSharedRef<IPropertyHandle> StructPropertyHandle, IDetailChildrenBuilder& ChildBuilder, IPropertyTypeCustomizationUtils& StructCustomizationUtils) override;
+		UE_API virtual void CustomizeHeader(TSharedRef<IPropertyHandle> StructPropertyHandle, FDetailWidgetRow& HeaderRow, IPropertyTypeCustomizationUtils& CustomizationUtils) override;
+		UE_API virtual void CustomizeChildren(TSharedRef<IPropertyHandle> StructPropertyHandle, IDetailChildrenBuilder& ChildBuilder, IPropertyTypeCustomizationUtils& StructCustomizationUtils) override;
 	};
 
-	class WAVETABLEEDITOR_API FTransformLayoutCustomizationBase : public IPropertyTypeCustomization
+	class FTransformLayoutCustomizationBase : public IPropertyTypeCustomization
 	{
 	public:
 		//~ Begin IPropertyTypeCustomization
-		virtual void CustomizeHeader(TSharedRef<IPropertyHandle> StructPropertyHandle, FDetailWidgetRow& HeaderRow, IPropertyTypeCustomizationUtils& StructCustomizationUtils) override;
-		virtual void CustomizeChildren(TSharedRef<IPropertyHandle> StructPropertyHandle, IDetailChildrenBuilder& ChildBuilder, IPropertyTypeCustomizationUtils& StructCustomizationUtils) override;
+		UE_API virtual void CustomizeHeader(TSharedRef<IPropertyHandle> StructPropertyHandle, FDetailWidgetRow& HeaderRow, IPropertyTypeCustomizationUtils& StructCustomizationUtils) override;
+		UE_API virtual void CustomizeChildren(TSharedRef<IPropertyHandle> StructPropertyHandle, IDetailChildrenBuilder& ChildBuilder, IPropertyTypeCustomizationUtils& StructCustomizationUtils) override;
 		//~ End IPropertyTypeCustomization
 
 	protected:
-		virtual TSet<EWaveTableCurve> GetSupportedCurves() const;
+		UE_API virtual TSet<EWaveTableCurve> GetSupportedCurves() const;
 		virtual FWaveTableTransform* GetTransform() const = 0;
 		virtual bool IsBipolar() const = 0;
-		virtual EWaveTableSamplingMode GetSamplingMode() const;
+		UE_API virtual EWaveTableSamplingMode GetSamplingMode() const;
 
-		void CachePCMFromFile();
-		EWaveTableCurve GetCurve() const;
-		int32 GetOwningArrayIndex() const;
-		bool IsScaleableCurve() const;
+		UE_API void CachePCMFromFile();
+		UE_API EWaveTableCurve GetCurve() const;
+		UE_API int32 GetOwningArrayIndex() const;
+		UE_API bool IsScaleableCurve() const;
 
 		TSharedPtr<IPropertyHandle> CurveHandle;
 		TSharedPtr<IPropertyHandle> ChannelIndexHandle;
@@ -50,12 +52,12 @@ namespace WaveTable::Editor
 		TSharedPtr<IPropertyHandle> WaveTableOptionsHandle;
 
 	private:
-		void CustomizeCurveSelector(IDetailChildrenBuilder& ChildBuilder);
+		UE_API void CustomizeCurveSelector(IDetailChildrenBuilder& ChildBuilder);
 
 		TMap<FString, FName> CurveDisplayStringToNameMap;
 	};
 
-	class WAVETABLEEDITOR_API FTransformLayoutCustomization : public FTransformLayoutCustomizationBase
+	class FTransformLayoutCustomization : public FTransformLayoutCustomizationBase
 	{
 	public:
 		static TSharedRef<IPropertyTypeCustomization> MakeInstance()
@@ -63,7 +65,9 @@ namespace WaveTable::Editor
 			return MakeShared<FTransformLayoutCustomization>();
 		}
 
-		virtual bool IsBipolar() const override;
-		virtual FWaveTableTransform* GetTransform() const override;
+		UE_API virtual bool IsBipolar() const override;
+		UE_API virtual FWaveTableTransform* GetTransform() const override;
 	};
 } // namespace WaveTable::Editor
+
+#undef UE_API
diff --git a/Engine/Plugins/Runtime/ZoneGraph/Source/ZoneGraphEditor/Private/ZoneShapeComponentVisualizer.h b/Engine/Plugins/Runtime/ZoneGraph/Source/ZoneGraphEditor/Private/ZoneShapeComponentVisualizer.h
--- a/Engine/Plugins/Runtime/ZoneGraph/Source/ZoneGraphEditor/Private/ZoneShapeComponentVisualizer.h
+++ b/Engine/Plugins/Runtime/ZoneGraph/Source/ZoneGraphEditor/Private/ZoneShapeComponentVisualizer.h
@@ -11,6 +11,8 @@
 #include "ZoneShapeComponent.h"
 #include "ZoneShapeComponentVisualizer.generated.h"
 
+#define UE_API ZONEGRAPHEDITOR_API
+
 class AActor;
 class FEditorViewportClient;
 class FMenuBuilder;
@@ -31,8 +33,8 @@ enum class FZoneShapeControlPointType : uint8
 };
 
 /** Selection state data that will be captured by scoped transactions.*/
-UCLASS(Transient)
-class ZONEGRAPHEDITOR_API UZoneShapeComponentVisualizerSelectionState : public UObject
+UCLASS(MinimalAPI, Transient)
+class UZoneShapeComponentVisualizerSelectionState : public UObject
 {
 	GENERATED_BODY()
 public:
@@ -81,9 +83,9 @@ protected:
 };
 
 /** Base class for clickable shape editing proxies */
-struct ZONEGRAPHEDITOR_API HZoneShapeVisProxy : public HComponentVisProxy
+struct HZoneShapeVisProxy : public HComponentVisProxy
 {
-	DECLARE_HIT_PROXY();
+	DECLARE_HIT_PROXY( UE_API );
 
 	HZoneShapeVisProxy(const UActorComponent* InComponent, EHitProxyPriority InPriority = HPP_Wireframe)
 		: HComponentVisProxy(InComponent, InPriority)
@@ -96,9 +98,9 @@ struct ZONEGRAPHEDITOR_API HZoneShapeVisProxy : public HComponentVisProxy
 };
 
 /** Proxy for a shape point */
-struct ZONEGRAPHEDITOR_API HZoneShapePointProxy : public HZoneShapeVisProxy
+struct HZoneShapePointProxy : public HZoneShapeVisProxy
 {
-	DECLARE_HIT_PROXY();
+	DECLARE_HIT_PROXY( UE_API );
 
 	HZoneShapePointProxy(const UActorComponent* InComponent, int32 InPointIndex)
 		: HZoneShapeVisProxy(InComponent, HPP_Foreground)
@@ -114,9 +116,9 @@ struct ZONEGRAPHEDITOR_API HZoneShapePointProxy : public HZoneShapeVisProxy
 };
 
 /** Proxy for a shape segment */
-struct ZONEGRAPHEDITOR_API HZoneShapeSegmentProxy : public HZoneShapeVisProxy
+struct HZoneShapeSegmentProxy : public HZoneShapeVisProxy
 {
-	DECLARE_HIT_PROXY();
+	DECLARE_HIT_PROXY( UE_API );
 
 	HZoneShapeSegmentProxy(const UActorComponent* InComponent, int32 InSegmentIndex)
 		: HZoneShapeVisProxy(InComponent)
@@ -132,9 +134,9 @@ struct ZONEGRAPHEDITOR_API HZoneShapeSegmentProxy : public HZoneShapeVisProxy
 };
 
 /** Proxy for a control point handle */
-struct ZONEGRAPHEDITOR_API HZoneShapeControlPointProxy : public HZoneShapeVisProxy
+struct HZoneShapeControlPointProxy : public HZoneShapeVisProxy
 {
-	DECLARE_HIT_PROXY();
+	DECLARE_HIT_PROXY( UE_API );
 
 	HZoneShapeControlPointProxy(const UActorComponent* InComponent, int32 InPointIndex, bool bInInControlPoint)
 		: HZoneShapeVisProxy(InComponent, HPP_Foreground)
@@ -152,38 +154,38 @@ struct ZONEGRAPHEDITOR_API HZoneShapeControlPointProxy : public HZoneShapeVisPro
 };
 
 /** ZoneShapeComponent visualizer/edit functionality */
-class ZONEGRAPHEDITOR_API FZoneShapeComponentVisualizer : public FComponentVisualizer, public FGCObject
+class FZoneShapeComponentVisualizer : public FComponentVisualizer, public FGCObject
 {
 public:
-	FZoneShapeComponentVisualizer();
-	virtual ~FZoneShapeComponentVisualizer();
+	UE_API FZoneShapeComponentVisualizer();
+	UE_API virtual ~FZoneShapeComponentVisualizer();
 
 	//~ Begin FComponentVisualizer Interface
-	virtual void OnRegister() override;
-	virtual void DrawVisualization(const UActorComponent* Component, const FSceneView* View, FPrimitiveDrawInterface* PDI) override;
-	virtual bool VisProxyHandleClick(FEditorViewportClient* InViewportClient, HComponentVisProxy* VisProxy, const FViewportClick& Click) override;
-	virtual void DrawVisualizationHUD(const UActorComponent* Component, const FViewport* Viewport, const FSceneView* View, FCanvas* Canvas) override;
-	virtual void EndEditing() override;
-	virtual bool GetWidgetLocation(const FEditorViewportClient* ViewportClient, FVector& OutLocation) const override;
-	virtual bool GetCustomInputCoordinateSystem(const FEditorViewportClient* ViewportClient, FMatrix& OutMatrix) const override;
-	virtual bool HandleInputDelta(FEditorViewportClient* ViewportClient, FViewport* Viewport, FVector& DeltaTranslate, FRotator& DeltaRotate, FVector& DeltaScale) override;
-	virtual bool HandleInputKey(FEditorViewportClient* ViewportClient, FViewport* Viewport, FKey Key, EInputEvent Event) override;
+	UE_API virtual void OnRegister() override;
+	UE_API virtual void DrawVisualization(const UActorComponent* Component, const FSceneView* View, FPrimitiveDrawInterface* PDI) override;
+	UE_API virtual bool VisProxyHandleClick(FEditorViewportClient* InViewportClient, HComponentVisProxy* VisProxy, const FViewportClick& Click) override;
+	UE_API virtual void DrawVisualizationHUD(const UActorComponent* Component, const FViewport* Viewport, const FSceneView* View, FCanvas* Canvas) override;
+	UE_API virtual void EndEditing() override;
+	UE_API virtual bool GetWidgetLocation(const FEditorViewportClient* ViewportClient, FVector& OutLocation) const override;
+	UE_API virtual bool GetCustomInputCoordinateSystem(const FEditorViewportClient* ViewportClient, FMatrix& OutMatrix) const override;
+	UE_API virtual bool HandleInputDelta(FEditorViewportClient* ViewportClient, FViewport* Viewport, FVector& DeltaTranslate, FRotator& DeltaRotate, FVector& DeltaScale) override;
+	UE_API virtual bool HandleInputKey(FEditorViewportClient* ViewportClient, FViewport* Viewport, FKey Key, EInputEvent Event) override;
 	/** Handle box select input */
-	virtual bool HandleBoxSelect(const FBox& InBox, FEditorViewportClient* InViewportClient, FViewport* InViewport) override;
+	UE_API virtual bool HandleBoxSelect(const FBox& InBox, FEditorViewportClient* InViewportClient, FViewport* InViewport) override;
 	/** Handle frustum select input */
-	virtual bool HandleFrustumSelect(const FConvexVolume& InFrustum, FEditorViewportClient* InViewportClient, FViewport* InViewport) override;
+	UE_API virtual bool HandleFrustumSelect(const FConvexVolume& InFrustum, FEditorViewportClient* InViewportClient, FViewport* InViewport) override;
 	/** Return whether focus on selection should focus on bounding box defined by active visualizer */
-	virtual bool HasFocusOnSelectionBoundingBox(FBox& OutBoundingBox) override;
+	UE_API virtual bool HasFocusOnSelectionBoundingBox(FBox& OutBoundingBox) override;
 	/** Pass snap input to active visualizer */
-	virtual bool HandleSnapTo(const bool bInAlign, const bool bInUseLineTrace, const bool bInUseBounds, const bool bInUsePivot, AActor* InDestination) override;
+	UE_API virtual bool HandleSnapTo(const bool bInAlign, const bool bInUseLineTrace, const bool bInUseBounds, const bool bInUsePivot, AActor* InDestination) override;
 	/** Get currently edited component, this is needed to reset the active visualizer after undo/redo */
-	virtual UActorComponent* GetEditedComponent() const override;
-	virtual TSharedPtr<SWidget> GenerateContextMenu() const override;
-	virtual bool IsVisualizingArchetype() const override;
+	UE_API virtual UActorComponent* GetEditedComponent() const override;
+	UE_API virtual TSharedPtr<SWidget> GenerateContextMenu() const override;
+	UE_API virtual bool IsVisualizingArchetype() const override;
 	//~ End FComponentVisualizer Interface
 
 	/** Get the shape component we are currently editing */
-	UZoneShapeComponent* GetEditedShapeComponent() const;
+	UE_API UZoneShapeComponent* GetEditedShapeComponent() const;
 
 	const TSet<int32>& GetSelectedPoints() const
 	{
@@ -214,70 +216,70 @@ protected:
 	};
 
 	/** Determine if any selected key index is out of range (perhaps because something external has modified the shape) */
-	bool IsAnySelectedPointIndexOutOfRange(const UZoneShapeComponent& Comp) const;
+	UE_API bool IsAnySelectedPointIndexOutOfRange(const UZoneShapeComponent& Comp) const;
 
 	/** Whether a single point is currently selected */
-	bool IsSinglePointSelected() const;
+	UE_API bool IsSinglePointSelected() const;
 
 	/** Transforms selected control point by given translation */
-	bool TransformSelectedControlPoint(const FVector& DeltaTranslate);
+	UE_API bool TransformSelectedControlPoint(const FVector& DeltaTranslate);
 
 	/** Transforms selected points by given delta translation, rotation and scale. */
-	bool TransformSelectedPoints(const FEditorViewportClient* ViewportClient, const FVector& DeltaTranslate, const FRotator& DeltaRotate, const FVector& DeltaScale) const;
+	UE_API bool TransformSelectedPoints(const FEditorViewportClient* ViewportClient, const FVector& DeltaTranslate, const FRotator& DeltaRotate, const FVector& DeltaScale) const;
 
 	/** Update the key selection state of the visualizer */
-	void ChangeSelectionState(int32 Index, bool bIsCtrlHeld) const;
+	UE_API void ChangeSelectionState(int32 Index, bool bIsCtrlHeld) const;
 
 	/** Alt-drag: duplicates the selected point */
-	bool DuplicatePointForAltDrag(const FVector& InDrag) const;
+	UE_API bool DuplicatePointForAltDrag(const FVector& InDrag) const;
 
 	/** Split segment using given interpolation value, selects the new point */
-	void SplitSegment(const int32 InSegmentIndex, const float SegmentSplitT, UZoneShapeComponent* ShapeComp = nullptr) const;
+	UE_API void SplitSegment(const int32 InSegmentIndex, const float SegmentSplitT, UZoneShapeComponent* ShapeComp = nullptr) const;
 
 	/** Add segment using given position and index, selects the new point */
-	void AddSegment(const FVector& InWorldPos, const int32 InSelectedIndex, UZoneShapeComponent* InShapeComp = nullptr) const;
+	UE_API void AddSegment(const FVector& InWorldPos, const int32 InSelectedIndex, UZoneShapeComponent* InShapeComp = nullptr) const;
 
 	/** Duplicates selected points and selects them. */
-	void DuplicateSelectedPoints(const FVector& WorldOffset = FVector::ZeroVector, bool bInsertAfter = true) const;
+	UE_API void DuplicateSelectedPoints(const FVector& WorldOffset = FVector::ZeroVector, bool bInsertAfter = true) const;
 
 	/** Updates the component and selected properties if the component has changed */
-	const UZoneShapeComponent* UpdateSelectedShapeComponent(const HComponentVisProxy* VisProxy);
+	UE_API const UZoneShapeComponent* UpdateSelectedShapeComponent(const HComponentVisProxy* VisProxy);
 
 	/** Returns the rotation of last selected point, or false if last selection is not valid. */
-	bool GetLastSelectedPointRotation(FQuat& OutRotation) const;
+	UE_API bool GetLastSelectedPointRotation(FQuat& OutRotation) const;
 
-	void OnDeletePoint() const;
-	bool CanDeletePoint() const;
+	UE_API void OnDeletePoint() const;
+	UE_API bool CanDeletePoint() const;
 
 	/** Duplicates selected points in place */
-	void OnDuplicatePoint() const;
-	bool IsPointSelectionValid() const;
+	UE_API void OnDuplicatePoint() const;
+	UE_API bool IsPointSelectionValid() const;
 
-	void OnAddPointToSegment() const;
-	bool CanAddPointToSegment() const;
+	UE_API void OnAddPointToSegment() const;
+	UE_API bool CanAddPointToSegment() const;
 
-	void OnSetPointType(FZoneShapePointType Type) const;
-	bool IsPointTypeSet(FZoneShapePointType Type) const;
+	UE_API void OnSetPointType(FZoneShapePointType Type) const;
+	UE_API bool IsPointTypeSet(FZoneShapePointType Type) const;
 
-	void OnSelectAllPoints() const;
-	bool CanSelectAllPoints() const;
+	UE_API void OnSelectAllPoints() const;
+	UE_API bool CanSelectAllPoints() const;
 
-	void OnBreakAtPointNewActors() const;
-	void OnBreakAtPointNewComponents() const;
-	bool CanBreakAtPoint() const;
+	UE_API void OnBreakAtPointNewActors() const;
+	UE_API void OnBreakAtPointNewComponents() const;
+	UE_API bool CanBreakAtPoint() const;
 
-	void OnBreakAtSegmentNewActors() const;
-	void OnBreakAtSegmentNewComponents() const;
-	bool CanBreakAtSegment() const;
+	UE_API void OnBreakAtSegmentNewActors() const;
+	UE_API void OnBreakAtSegmentNewComponents() const;
+	UE_API bool CanBreakAtSegment() const;
 
-	void GenerateShapePointTypeSubMenu(FMenuBuilder& MenuBuilder) const;
+	UE_API void GenerateShapePointTypeSubMenu(FMenuBuilder& MenuBuilder) const;
 
-	void GenerateSnapAlignSubMenu(FMenuBuilder& MenuBuilder) const;
-	void GenerateBreakAtPointSubMenu(FMenuBuilder& MenuBuilder) const;
-	void GenerateBreakAtSegmentSubMenu(FMenuBuilder& MenuBuilder) const;
+	UE_API void GenerateSnapAlignSubMenu(FMenuBuilder& MenuBuilder) const;
+	UE_API void GenerateBreakAtPointSubMenu(FMenuBuilder& MenuBuilder) const;
+	UE_API void GenerateBreakAtSegmentSubMenu(FMenuBuilder& MenuBuilder) const;
 
 	// FGCObject interface
-	virtual void AddReferencedObjects(FReferenceCollector& Collector) override;
+	UE_API virtual void AddReferencedObjects(FReferenceCollector& Collector) override;
 	virtual FString GetReferencerName() const override
 	{
 		return TEXT("FZoneShapeComponentVisualizer");
@@ -339,19 +341,21 @@ protected:
 
 private:
 
-	TArray<UZoneShapeComponent*>  BreakAtPoint(bool bCreateNewActor, UZoneShapeComponent* ShapeComp = nullptr) const;
-	void BreakAtSegment(bool bCreateNewActor) const;
+	UE_API TArray<UZoneShapeComponent*>  BreakAtPoint(bool bCreateNewActor, UZoneShapeComponent* ShapeComp = nullptr) const;
+	UE_API void BreakAtSegment(bool bCreateNewActor) const;
 
-	void DetectCloseByShapeForAutoConnection(const UZoneShapeComponent* ShapeComp, const FZoneShapePoint& DraggedPoint);
-	void DetectCloseByShapeForAutoIntersectionCreation(const UZoneShapeComponent* ShapeComp, const FZoneShapePoint& DraggedPoint);
+	UE_API void DetectCloseByShapeForAutoConnection(const UZoneShapeComponent* ShapeComp, const FZoneShapePoint& DraggedPoint);
+	UE_API void DetectCloseByShapeForAutoIntersectionCreation(const UZoneShapeComponent* ShapeComp, const FZoneShapePoint& DraggedPoint);
 
-	void CreateIntersection(UZoneShapeComponent* ShapeComp);
-	void CreateIntersectionForSplineShape(UZoneShapeComponent* ShapeComp, FZoneShapePoint& DraggedPoint, bool DestroyCoveredShape = true);
-	void CreateIntersectionForPolygonShape(UZoneShapeComponent* ShapeComp, FZoneShapePoint& DraggedPoint);
+	UE_API void CreateIntersection(UZoneShapeComponent* ShapeComp);
+	UE_API void CreateIntersectionForSplineShape(UZoneShapeComponent* ShapeComp, FZoneShapePoint& DraggedPoint, bool DestroyCoveredShape = true);
+	UE_API void CreateIntersectionForPolygonShape(UZoneShapeComponent* ShapeComp, FZoneShapePoint& DraggedPoint);
 
-	void ClearAutoConnectingStatus();
-	void ClearAutoIntersectionStatus();
+	UE_API void ClearAutoConnectingStatus();
+	UE_API void ClearAutoIntersectionStatus();
 
-	bool CanAutoConnect(const UZoneShapeComponent* ShapeComp) const;
-	bool CanAutoCreateIntersection(const UZoneShapeComponent* ShapeComp) const;
+	UE_API bool CanAutoConnect(const UZoneShapeComponent* ShapeComp) const;
+	UE_API bool CanAutoCreateIntersection(const UZoneShapeComponent* ShapeComp) const;
 };
+
+#undef UE_API
diff --git a/Engine/Plugins/Runtime/ZoneGraph/Source/ZoneGraphEditor/Public/IZoneGraphEditor.h b/Engine/Plugins/Runtime/ZoneGraph/Source/ZoneGraphEditor/Public/IZoneGraphEditor.h
--- a/Engine/Plugins/Runtime/ZoneGraph/Source/ZoneGraphEditor/Public/IZoneGraphEditor.h
+++ b/Engine/Plugins/Runtime/ZoneGraph/Source/ZoneGraphEditor/Public/IZoneGraphEditor.h
@@ -7,7 +7,7 @@
 
 
 /** ZoneGraph Editor public interface */
-class ZONEGRAPHEDITOR_API IZoneGraphEditor : public FAssetEditorToolkit
+class IZoneGraphEditor : public FAssetEditorToolkit
 {
 public:
 
diff --git a/Engine/Plugins/Runtime/ZoneGraph/Source/ZoneGraphEditor/Public/ZoneGraphEditorModule.h b/Engine/Plugins/Runtime/ZoneGraph/Source/ZoneGraphEditor/Public/ZoneGraphEditorModule.h
--- a/Engine/Plugins/Runtime/ZoneGraph/Source/ZoneGraphEditor/Public/ZoneGraphEditorModule.h
+++ b/Engine/Plugins/Runtime/ZoneGraph/Source/ZoneGraphEditor/Public/ZoneGraphEditorModule.h
@@ -12,26 +12,20 @@ class FUICommandList;
 /**
 * The public interface to this module
 */
-class ZONEGRAPHEDITOR_API FZoneGraphEditorModule : public IModuleInterface
+class FZoneGraphEditorModule : public IModuleInterface
 {
 public:
 	// Begin IModuleInterface
-	virtual void StartupModule() override;
-	virtual void ShutdownModule() override;
+	ZONEGRAPHEDITOR_API virtual void StartupModule() override;
+	ZONEGRAPHEDITOR_API virtual void ShutdownModule() override;
 
 private:
-	void RegisterComponentVisualizer(FName ComponentClassName, TSharedPtr<FComponentVisualizer> Visualizer);
-	void RegisterMenus();
+	ZONEGRAPHEDITOR_API void RegisterComponentVisualizer(FName ComponentClassName, TSharedPtr<FComponentVisualizer> Visualizer);
+	ZONEGRAPHEDITOR_API void RegisterMenus();
 
-	void OnBuildZoneGraph();
+	ZONEGRAPHEDITOR_API void OnBuildZoneGraph();
 
 	TArray<FName> RegisteredComponentClassNames;
 	TSharedPtr<FUICommandList> PluginCommands;
 };
 
-#if UE_ENABLE_INCLUDE_ORDER_DEPRECATED_IN_5_3
-#include "CoreMinimal.h"
-#include "Modules/ModuleManager.h"
-#include "Toolkits/IToolkitHost.h"
-#include "Toolkits/AssetEditorToolkit.h"
-#endif
\ No newline at end of file
