Zone Graph: Fixed a crash on creating an intersection with X+drag shortcut. Fixed a crash when using V shortcut without having a point selected.

#rnx

diff --git a/Engine/Plugins/Runtime/ZoneGraph/Source/ZoneGraphEditor/Private/ZoneShapeComponentVisualizer.cpp b/Engine/Plugins/Runtime/ZoneGraph/Source/ZoneGraphEditor/Private/ZoneShapeComponentVisualizer.cpp
--- a/Engine/Plugins/Runtime/ZoneGraph/Source/ZoneGraphEditor/Private/ZoneShapeComponentVisualizer.cpp
+++ b/Engine/Plugins/Runtime/ZoneGraph/Source/ZoneGraphEditor/Private/ZoneShapeComponentVisualizer.cpp
@@ -1552,8 +1552,10 @@ bool FZoneShapeComponentVisualizer::HandleInputKey(FEditorViewportClient* Viewpo
 
 	if (Event == IE_Pressed)
 	{
+		const int32 SelectedIndex = SelectionState->GetLastPointIndexSelected();
+
 		// Add a new point to the shape when you hold the V key and press left mouse button
-		if (ShapeComp && Key == EKeys::LeftMouseButton && Viewport->KeyState(EKeys::V))
+		if (ShapeComp && Key == EKeys::LeftMouseButton && Viewport->KeyState(EKeys::V) && SelectedIndex != INDEX_NONE)
 		{
 			// Get clicked position
 			UWorld* World = ViewportClient->GetWorld();
@@ -1578,7 +1580,6 @@ bool FZoneShapeComponentVisualizer::HandleInputKey(FEditorViewportClient* Viewpo
 			{
 				// Add a new point at the position
 				const FScopedTransaction Transaction(LOCTEXT("AddShapePointAndSnap", "Add Shape Point And Snap To Floor"));
-				const int32 SelectedIndex = SelectionState->GetLastPointIndexSelected();
 				AddSegment(Hit.Location, SelectedIndex, ShapeComp);
 			}
 			else
@@ -2715,6 +2716,10 @@ void FZoneShapeComponentVisualizer::CreateIntersectionForSplineShape(UZoneShapeC
 		Index0 = Points0.Num() - 1;
 		if (ShapeComponents[0])
 		{
+			// Need to update the shape as points could be removed
+			// This will rebuild the shape connectors before the GetShapeConnectorByPointIndex call
+			ShapeComponents[0]->UpdateShape();
+			
 			FVector Normal0 = ShapeComponents[0]->GetShapeConnectorByPointIndex(Index0)->Normal;
 			FVector Offset = Normal0 * HalfLanesTotalWidth;
 			if (Points0.Num() == 2)
@@ -2741,6 +2746,10 @@ void FZoneShapeComponentVisualizer::CreateIntersectionForSplineShape(UZoneShapeC
 		}
 		if (ShapeComponents[1])
 		{
+			// Need to update the shape as points could be removed
+			// This will rebuild the shape connectors before the GetShapeConnectorByPointIndex call
+			ShapeComponents[1]->UpdateShape();
+			
 			FVector Normal1 = ShapeComponents[1]->GetShapeConnectorByPointIndex(Index1)->Normal;
 			FVector Offset = Normal1 * HalfLanesTotalWidth;
 			if (Points1.Num() == 2)
