PR #13526: Fix problem with converting FZoneGraphTag::None to FZoneGraphTagMask

diff --git a/Engine/Plugins/Runtime/ZoneGraph/Source/ZoneGraph/Public/ZoneGraphTypes.h b/Engine/Plugins/Runtime/ZoneGraph/Source/ZoneGraph/Public/ZoneGraphTypes.h
--- a/Engine/Plugins/Runtime/ZoneGraph/Source/ZoneGraph/Public/ZoneGraphTypes.h
+++ b/Engine/Plugins/Runtime/ZoneGraph/Source/ZoneGraph/Public/ZoneGraphTypes.h
@@ -73,12 +73,20 @@ struct FZoneGraphTag
 		Bit = NoneValue;
 	}
 
-	bool IsValid() const {
+	bool IsValid() const
+	{
 		return Bit != NoneValue;
 	}
 
-	bool operator==(const FZoneGraphTag& RHS) const { return Bit == RHS.Bit; }
-	bool operator!=(const FZoneGraphTag& RHS) const { return Bit != RHS.Bit; }
+	bool operator==(const FZoneGraphTag& RHS) const
+	{
+		return Bit == RHS.Bit;
+	}
+
+	bool operator!=(const FZoneGraphTag& RHS) const
+	{
+		return Bit != RHS.Bit;
+	}
 
 	friend uint32 GetTypeHash(const FZoneGraphTag& Tag)
 	{
@@ -111,7 +119,7 @@ struct FZoneGraphTagMask
 	static UE_API const FZoneGraphTagMask None;
 
 	FZoneGraphTagMask() = default;
-	FZoneGraphTagMask(const FZoneGraphTag InTag) : Mask(uint32(1) << InTag.Bit) {}
+	FZoneGraphTagMask(const FZoneGraphTag InTag) : Mask(InTag.IsValid() ? uint32(1) << InTag.Bit : 0) {}
 	explicit FZoneGraphTagMask(uint32 InMask) : Mask(InMask) {}
 
 	void Add(const FZoneGraphTagMask InTags)
