Zone Graph: Add an extra FZoneDrawAnnotator parameter to be able to customize zone graph draw debugs.

#rnx
#Synced-CL 48147785

diff --git a/Engine/Plugins/Runtime/ZoneGraph/Source/ZoneGraph/Private/ZoneGraphRenderingUtilities.cpp b/Engine/Plugins/Runtime/ZoneGraph/Source/ZoneGraph/Private/ZoneGraphRenderingUtilities.cpp
--- a/Engine/Plugins/Runtime/ZoneGraph/Source/ZoneGraph/Private/ZoneGraphRenderingUtilities.cpp
+++ b/Engine/Plugins/Runtime/ZoneGraph/Source/ZoneGraph/Private/ZoneGraphRenderingUtilities.cpp
@@ -29,7 +29,7 @@ FColor GetTagMaskColor(const FZoneGraphTagMask TagMask, TConstArrayView<FZoneGra
 	return Weight > 0.0f ? (Color / Weight).ToFColor(true /*bSRGB*/) : FColor::Black;
 }
 
-void DrawZoneBoundary(const FZoneGraphStorage& ZoneStorage, int32 ZoneIndex, FPrimitiveDrawInterface* PDI, const FMatrix& LocalToWorld, const float LineThickness, const float DepthBias, const float Alpha)
+void DrawZoneBoundary(const FZoneGraphStorage& ZoneStorage, int32 ZoneIndex, FPrimitiveDrawInterface* PDI, const FMatrix& LocalToWorld, const float LineThickness, const float DepthBias, const float Alpha, const FZoneDrawAnnotator& Annotator)
 {
 	if (ZoneIndex < 0 || ZoneIndex >= ZoneStorage.Zones.Num())
 	{
@@ -38,12 +38,12 @@ void DrawZoneBoundary(const FZoneGraphStorage& ZoneStorage, int32 ZoneIndex, FPr
 
 	const UZoneGraphSettings* ZoneGraphSettings = GetDefault<UZoneGraphSettings>();
 	check(ZoneGraphSettings);
-	FZoneGraphTagMask VisTagMask = ZoneGraphSettings->GetVisualizedTags();
+	FZoneGraphTagMask VisTagMask = Annotator.VisualizedTagsOverride.Get(ZoneGraphSettings->GetVisualizedTags());
 	TConstArrayView<FZoneGraphTagInfo>TagInfo = ZoneGraphSettings->GetTagInfos();
 
 	const FZoneData& Zone = ZoneStorage.Zones[ZoneIndex];
 	const bool bTransform = !LocalToWorld.Equals(FMatrix::Identity);
-	FColor Color = GetTagMaskColor(Zone.Tags & VisTagMask, TagInfo);
+	FColor Color = GetTagMaskColor(Annotator.ZoneTagsOverride.Get(Zone.Tags) & VisTagMask, TagInfo);
 	Color.A = uint8(128 * Alpha);
 
 	FVector PrevPoint = bTransform ? FVector(LocalToWorld.TransformPosition(ZoneStorage.BoundaryPoints[Zone.BoundaryPointsEnd - 1])) : ZoneStorage.BoundaryPoints[Zone.BoundaryPointsEnd - 1];
@@ -273,7 +273,7 @@ void DrawLanePath(const FZoneGraphStorage& ZoneStorage, FPrimitiveDrawInterface*
 	}
 }
 
-void DrawZoneLanes(const FZoneGraphStorage& ZoneStorage, int32 ZoneIndex, FPrimitiveDrawInterface* PDI, const FMatrix& LocalToWorld, const float LineThickness, const float DepthBias, const float Alpha, const bool bDrawDetails, const FLaneHighlight& LaneHighlight)
+void DrawZoneLanes(const FZoneGraphStorage& ZoneStorage, int32 ZoneIndex, FPrimitiveDrawInterface* PDI, const FMatrix& LocalToWorld, const float LineThickness, const float DepthBias, const float Alpha, const bool bDrawDetails, const FLaneHighlight& LaneHighlight, const FZoneDrawAnnotator& Annotator)
 {
 	if (ZoneIndex < 0 || ZoneIndex >= ZoneStorage.Zones.Num())
 	{
@@ -282,7 +282,7 @@ void DrawZoneLanes(const FZoneGraphStorage& ZoneStorage, int32 ZoneIndex, FPrimi
 
 	const UZoneGraphSettings* ZoneGraphSettings = GetDefault<UZoneGraphSettings>();
 	check(ZoneGraphSettings);
-	FZoneGraphTagMask VisTagMask = ZoneGraphSettings->GetVisualizedTags();
+	FZoneGraphTagMask VisTagMask = Annotator.VisualizedTagsOverride.Get(ZoneGraphSettings->GetVisualizedTags());
 	TConstArrayView<FZoneGraphTagInfo>TagInfo = ZoneGraphSettings->GetTagInfos();
 
 	const FZoneData& Zone = ZoneStorage.Zones[ZoneIndex];
@@ -292,7 +292,7 @@ void DrawZoneLanes(const FZoneGraphStorage& ZoneStorage, int32 ZoneIndex, FPrimi
 	for (int32 LaneIdx = Zone.LanesBegin; LaneIdx < Zone.LanesEnd; LaneIdx++)
 	{
 		const FZoneLaneData& Lane = ZoneStorage.Lanes[LaneIdx];
-		FColor Color = GetTagMaskColor(Lane.Tags & VisTagMask, TagInfo);
+		FColor Color = GetTagMaskColor((Annotator.GetLaneTagsOverride.IsSet() ? Annotator.GetLaneTagsOverride.GetValue()(LaneIdx, Lane) : Lane.Tags) & VisTagMask, TagInfo);
 		bool bDrawTransparent = false;
 
 		FVector PrevPoint = bTransform ? FVector(LocalToWorld.TransformPosition(ZoneStorage.LanePoints[Lane.PointsBegin])) : ZoneStorage.LanePoints[Lane.PointsBegin];
diff --git a/Engine/Plugins/Runtime/ZoneGraph/Source/ZoneGraph/Public/ZoneGraphRenderingUtilities.h b/Engine/Plugins/Runtime/ZoneGraph/Source/ZoneGraph/Public/ZoneGraphRenderingUtilities.h
--- a/Engine/Plugins/Runtime/ZoneGraph/Source/ZoneGraph/Public/ZoneGraphRenderingUtilities.h
+++ b/Engine/Plugins/Runtime/ZoneGraph/Source/ZoneGraph/Public/ZoneGraphRenderingUtilities.h
@@ -20,14 +20,21 @@ struct FLaneHighlight
 	float Width = 0.0f;
 };
 
+struct FZoneDrawAnnotator
+{
+	TOptional<FZoneGraphTagMask> VisualizedTagsOverride;
+	TOptional<FZoneGraphTagMask> ZoneTagsOverride;
+	TOptional<TFunctionRef<FZoneGraphTagMask(int32 LaneIndex, const FZoneLaneData& Lane)>> GetLaneTagsOverride;
+};
+
 ZONEGRAPH_API FColor GetTagMaskColor(const FZoneGraphTagMask TagMask, TConstArrayView<FZoneGraphTagInfo> TagInfo);
-ZONEGRAPH_API void DrawZoneBoundary(const FZoneGraphStorage& ZoneStorage, int32 ZoneIndex, FPrimitiveDrawInterface* PDI, const FMatrix& LocalToWorld, const float LineThickness, const float DepthBias, const float Alpha);
+ZONEGRAPH_API void DrawZoneBoundary(const FZoneGraphStorage& ZoneStorage, int32 ZoneIndex, FPrimitiveDrawInterface* PDI, const FMatrix& LocalToWorld, const float LineThickness, const float DepthBias, const float Alpha, const FZoneDrawAnnotator& Annotator = { });
 ZONEGRAPH_API void DrawLane(const FZoneGraphStorage& ZoneStorage, FPrimitiveDrawInterface* PDI, const FZoneGraphLaneHandle& LaneHandle, const FColor Color, const float LineThickness = 5.f, const FVector& Offset = FVector::ZeroVector);
 ZONEGRAPH_API void DrawLane(const FZoneGraphStorage& ZoneStorage, FPrimitiveDrawInterface* PDI, const FZoneGraphLaneLocation& InStartLocation, const FZoneGraphLaneLocation& InEndLocation, const FColor Color, const float LineThickness = 5.f, const FVector& Offset = FVector::ZeroVector);
 ZONEGRAPH_API void DrawLanes(const FZoneGraphStorage& ZoneStorage, FPrimitiveDrawInterface* PDI, const TArray<FZoneGraphLaneHandle>& Lanes, const FColor Color, const float LineThickness = 5.f, const FVector& Offset = FVector::ZeroVector);
 ZONEGRAPH_API void DrawLinkedLanes(const FZoneGraphStorage& ZoneStorage, FPrimitiveDrawInterface* PDI, const FZoneGraphLaneHandle& SourceLaneHandle, const TArray<FZoneGraphLinkedLane>& LinkedLanes, const float LineThickness = 5.f);
 ZONEGRAPH_API void DrawLanePath(const FZoneGraphStorage& ZoneStorage, FPrimitiveDrawInterface* PDI, const FZoneGraphLanePath& LanePath, const FColor Color, const float LineThickness = 5.f);
-ZONEGRAPH_API void DrawZoneLanes(const FZoneGraphStorage& ZoneStorage, int32 ZoneIndex, FPrimitiveDrawInterface* PDI, const FMatrix& LocalToWorld, const float LineThickness, const float DepthBias, const float Alpha, const bool bDrawDetails = true, const FLaneHighlight& LaneHighlight = FLaneHighlight());
+ZONEGRAPH_API void DrawZoneLanes(const FZoneGraphStorage& ZoneStorage, int32 ZoneIndex, FPrimitiveDrawInterface* PDI, const FMatrix& LocalToWorld, const float LineThickness, const float DepthBias, const float Alpha, const bool bDrawDetails = true, const FLaneHighlight& LaneHighlight = FLaneHighlight(), const FZoneDrawAnnotator& Annotator = { });
 
 ZONEGRAPH_API void DrawZoneShapeConnector(const FZoneShapeConnector& Connector, const FZoneShapeConnection* Connection, FPrimitiveDrawInterface* PDI, const FMatrix& LocalToWorld, const float DepthBias);
 ZONEGRAPH_API void DrawZoneShapeConnectors(TConstArrayView<FZoneShapeConnector> Connectors, TConstArrayView<FZoneShapeConnection> Connections, FPrimitiveDrawInterface* PDI, const FMatrix& LocalToWorld, const float DepthBias);
